Networking: added FIPv4Endpoint.ToString()

diff --git a/Engine/Source/Runtime/Networking/Private/IPv4/IPv4Endpoint.cpp b/Engine/Source/Runtime/Networking/Private/IPv4/IPv4Endpoint.cpp
--- a/Engine/Source/Runtime/Networking/Private/IPv4/IPv4Endpoint.cpp
+++ b/Engine/Source/Runtime/Networking/Private/IPv4/IPv4Endpoint.cpp
@@ -17,9 +17,9 @@ ISocketSubsystem* FIPv4Endpoint::CachedSocketSubsystem = NULL;
 /* FIPv4Endpoint interface
  *****************************************************************************/
 
-FText FIPv4Endpoint::ToText( ) const
+FString FIPv4Endpoint::ToString( ) const
 {
-	return FText::FromString( FString::Printf(TEXT("%s:%i"), *Address.ToText().ToString(), Port) );
+	return FString::Printf(TEXT("%s:%i"), *Address.ToText().ToString(), Port);
 }
 
 
diff --git a/Engine/Source/Runtime/Networking/Public/Interfaces/IPv4/IPv4Endpoint.h b/Engine/Source/Runtime/Networking/Public/Interfaces/IPv4/IPv4Endpoint.h
--- a/Engine/Source/Runtime/Networking/Public/Interfaces/IPv4/IPv4Endpoint.h
+++ b/Engine/Source/Runtime/Networking/Public/Interfaces/IPv4/IPv4Endpoint.h
@@ -134,11 +134,23 @@ public:
 	}
 
 	/**
-	 * Gets string representation for this endpoint.
+	 * Gets a string representation for this endpoint.
 	 *
 	 * @return String representation.
+	 *
+	 * @see Parse
+	 */
+	NETWORKING_API FString ToString( ) const;
+
+	/**
+	 * Gets the display text representation for this endpoint.
+	 *
+	 * @return Text representation.
 	 */
-	NETWORKING_API FText ToText( ) const;
+	FText ToText( ) const
+	{
+		return FText::FromString(ToString());
+	}
 
 public:
 
@@ -168,6 +180,8 @@ public:
 	 * @param OutEndpoint - Will contain the parsed endpoint.
 	 *
 	 * @return true if the string was converted successfully, false otherwise.
+	 *
+	 * @see ToString
 	 */
 	static NETWORKING_API bool Parse( const FString& EndpointString, FIPv4Endpoint& OutEndpoint );
 
