Docs: Removed file comments and added missing code documentation

diff --git a/Engine/Source/Runtime/Messaging/Public/Interfaces/IMessageContext.h b/Engine/Source/Runtime/Messaging/Public/Interfaces/IMessageContext.h
--- a/Engine/Source/Runtime/Messaging/Public/Interfaces/IMessageContext.h
+++ b/Engine/Source/Runtime/Messaging/Public/Interfaces/IMessageContext.h
@@ -1,9 +1,5 @@
 // Copyright 1998-2014 Epic Games, Inc. All Rights Reserved.
 
-/*=============================================================================
-	IMessageContext.h: Declares the IMessageContext interface.
-=============================================================================*/
-
 #pragma once
 
 
diff --git a/Engine/Source/Runtime/Messaging/Public/Interfaces/IMessageData.h b/Engine/Source/Runtime/Messaging/Public/Interfaces/IMessageData.h
--- a/Engine/Source/Runtime/Messaging/Public/Interfaces/IMessageData.h
+++ b/Engine/Source/Runtime/Messaging/Public/Interfaces/IMessageData.h
@@ -1,9 +1,5 @@
 // Copyright 1998-2014 Epic Games, Inc. All Rights Reserved.
 
-/*=============================================================================
-	IMessageAttachment.h: Declares the IMessageAttachment interface.
-=============================================================================*/
-
 #pragma once
 
 
diff --git a/Engine/Source/Runtime/Messaging/Public/Interfaces/IMessageHandler.h b/Engine/Source/Runtime/Messaging/Public/Interfaces/IMessageHandler.h
--- a/Engine/Source/Runtime/Messaging/Public/Interfaces/IMessageHandler.h
+++ b/Engine/Source/Runtime/Messaging/Public/Interfaces/IMessageHandler.h
@@ -1,9 +1,5 @@
 // Copyright 1998-2014 Epic Games, Inc. All Rights Reserved.
 
-/*=============================================================================
-	IMessageHandler.h: Declares the IMessageHandler interface and related helper templates.
-=============================================================================*/
-
 #pragma once
 
 
@@ -84,7 +80,7 @@ public:
 
 public:
 
-	// Begin IMessageHandler interface
+	// IMessageHandler interface
 	
 	virtual const FName GetHandledMessageType( ) const override
 	{
@@ -96,8 +92,6 @@ public:
 		(Handler->*HandlerFunc)(*static_cast<const MessageType*>(Context->GetMessage()), Context);
 	}
 	
-	// End IMessageHandler interface
-	
 private:
 
 	// Holds a pointer to the object handling the messages.
diff --git a/Engine/Source/Runtime/Messaging/Public/Interfaces/IMessageSubscription.h b/Engine/Source/Runtime/Messaging/Public/Interfaces/IMessageSubscription.h
--- a/Engine/Source/Runtime/Messaging/Public/Interfaces/IMessageSubscription.h
+++ b/Engine/Source/Runtime/Messaging/Public/Interfaces/IMessageSubscription.h
@@ -1,9 +1,5 @@
 // Copyright 1998-2014 Epic Games, Inc. All Rights Reserved.
 
-/*=============================================================================
-	IMessageSubscription.h: Declares the IMessageSubscription interface.
-=============================================================================*/
-
 #pragma once
 
 
diff --git a/Engine/Source/Runtime/Messaging/Public/Interfaces/IMessageTracer.h b/Engine/Source/Runtime/Messaging/Public/Interfaces/IMessageTracer.h
--- a/Engine/Source/Runtime/Messaging/Public/Interfaces/IMessageTracer.h
+++ b/Engine/Source/Runtime/Messaging/Public/Interfaces/IMessageTracer.h
@@ -1,9 +1,5 @@
 // Copyright 1998-2014 Epic Games, Inc. All Rights Reserved.
 
-/*=============================================================================
-	IMessageTracer.h: Declares the IMessageTracer interface.
-=============================================================================*/
-
 #pragma once
 
 
diff --git a/Engine/Source/Runtime/Messaging/Public/Interfaces/IMessageTracerBreakpoint.h b/Engine/Source/Runtime/Messaging/Public/Interfaces/IMessageTracerBreakpoint.h
--- a/Engine/Source/Runtime/Messaging/Public/Interfaces/IMessageTracerBreakpoint.h
+++ b/Engine/Source/Runtime/Messaging/Public/Interfaces/IMessageTracerBreakpoint.h
@@ -1,9 +1,5 @@
 // Copyright 1998-2014 Epic Games, Inc. All Rights Reserved.
 
-/*=============================================================================
-	IMessageTracerBreakpoint.h: Declares the IMessageTracerBreakpoint interface.
-=============================================================================*/
-
 #pragma once
 
 
diff --git a/Engine/Source/Runtime/Messaging/Public/Interfaces/IMessagingModule.h b/Engine/Source/Runtime/Messaging/Public/Interfaces/IMessagingModule.h
--- a/Engine/Source/Runtime/Messaging/Public/Interfaces/IMessagingModule.h
+++ b/Engine/Source/Runtime/Messaging/Public/Interfaces/IMessagingModule.h
@@ -1,9 +1,5 @@
 // Copyright 1998-2014 Epic Games, Inc. All Rights Reserved.
 
-/*=============================================================================
-	IMessagingModule.h: Declares the IMessagingModule interface.
-=============================================================================*/
-
 #pragma once
 
 
diff --git a/Engine/Source/Runtime/Messaging/Public/Interfaces/IMutableMessageContext.h b/Engine/Source/Runtime/Messaging/Public/Interfaces/IMutableMessageContext.h
--- a/Engine/Source/Runtime/Messaging/Public/Interfaces/IMutableMessageContext.h
+++ b/Engine/Source/Runtime/Messaging/Public/Interfaces/IMutableMessageContext.h
@@ -1,9 +1,5 @@
 // Copyright 1998-2014 Epic Games, Inc. All Rights Reserved.
 
-/*=============================================================================
-	IMutableMessageContext.h: Declares the IMutableMessageContext interface.
-=============================================================================*/
-
 #pragma once
 
 
diff --git a/Engine/Source/Runtime/Messaging/Public/Interfaces/IReceiveMessages.h b/Engine/Source/Runtime/Messaging/Public/Interfaces/IReceiveMessages.h
--- a/Engine/Source/Runtime/Messaging/Public/Interfaces/IReceiveMessages.h
+++ b/Engine/Source/Runtime/Messaging/Public/Interfaces/IReceiveMessages.h
@@ -1,9 +1,5 @@
 // Copyright 1998-2014 Epic Games, Inc. All Rights Reserved.
 
-/*=============================================================================
-	IReceiveMessages.h: Declares the IReceiveMessages interface.
-=============================================================================*/
-
 #pragma once
 
 
diff --git a/Engine/Source/Runtime/Messaging/Public/Interfaces/ISendMessages.h b/Engine/Source/Runtime/Messaging/Public/Interfaces/ISendMessages.h
--- a/Engine/Source/Runtime/Messaging/Public/Interfaces/ISendMessages.h
+++ b/Engine/Source/Runtime/Messaging/Public/Interfaces/ISendMessages.h
@@ -1,9 +1,5 @@
 // Copyright 1998-2014 Epic Games, Inc. All Rights Reserved.
 
-/*=============================================================================
-	ISendMessages.h: Declares the ISendMessages interface.
-=============================================================================*/
-
 #pragma once
 
 
diff --git a/Engine/Source/Runtime/Messaging/Public/Interfaces/ISerializeMessages.h b/Engine/Source/Runtime/Messaging/Public/Interfaces/ISerializeMessages.h
--- a/Engine/Source/Runtime/Messaging/Public/Interfaces/ISerializeMessages.h
+++ b/Engine/Source/Runtime/Messaging/Public/Interfaces/ISerializeMessages.h
@@ -1,9 +1,5 @@
 // Copyright 1998-2014 Epic Games, Inc. All Rights Reserved.
 
-/*=============================================================================
-	ISerializeMessages.h: Declares the ISerializeMessages interface.
-=============================================================================*/
-
 #pragma once
 
 
diff --git a/Engine/Source/Runtime/Messaging/Public/Interfaces/ITransportMessages.h b/Engine/Source/Runtime/Messaging/Public/Interfaces/ITransportMessages.h
--- a/Engine/Source/Runtime/Messaging/Public/Interfaces/ITransportMessages.h
+++ b/Engine/Source/Runtime/Messaging/Public/Interfaces/ITransportMessages.h
@@ -1,9 +1,5 @@
 // Copyright 1998-2014 Epic Games, Inc. All Rights Reserved.
 
-/*=============================================================================
-	ITransportMessages.h: Declares the ITransportMessages interface.
-=============================================================================*/
-
 #pragma once
 
 
diff --git a/Engine/Source/Runtime/Messaging/Public/Messaging.h b/Engine/Source/Runtime/Messaging/Public/Messaging.h
--- a/Engine/Source/Runtime/Messaging/Public/Messaging.h
+++ b/Engine/Source/Runtime/Messaging/Public/Messaging.h
@@ -1,9 +1,5 @@
 // Copyright 1998-2014 Epic Games, Inc. All Rights Reserved.
 
-/*=============================================================================
-	Messaging.h: Messaging module public header file.
-=============================================================================*/
-
 #pragma once
 
 
@@ -15,7 +11,7 @@
 #include "TaskGraphInterfaces.h"
 
 
-/* Public Interfaces
+/* Public Includes
  *****************************************************************************/
 
 #include "IMessageAttachment.h"
@@ -37,7 +33,7 @@
 #include "IMessagingModule.h"
 
 
-/* Public Common
+/* Common
  *****************************************************************************/
 
 #include "FileMessageAttachment.h"
diff --git a/Engine/Source/Runtime/NetworkFileSystem/Private/NetworkFileServer.cpp b/Engine/Source/Runtime/NetworkFileSystem/Private/NetworkFileServer.cpp
--- a/Engine/Source/Runtime/NetworkFileSystem/Private/NetworkFileServer.cpp
+++ b/Engine/Source/Runtime/NetworkFileSystem/Private/NetworkFileServer.cpp
@@ -1,9 +1,5 @@
 // Copyright 1998-2014 Epic Games, Inc. All Rights Reserved.
 
-/*=============================================================================
-	NetworkFileServer.cpp: Implements the FNetworkFileServer class.
-=============================================================================*/
-
 #include "NetworkFileSystemPrivatePCH.h"
 #include "TargetPlatform.h"
 #include "NetworkFileServerConnection.h"
diff --git a/Engine/Source/Runtime/NetworkFileSystem/Private/NetworkFileServer.h b/Engine/Source/Runtime/NetworkFileSystem/Private/NetworkFileServer.h
--- a/Engine/Source/Runtime/NetworkFileSystem/Private/NetworkFileServer.h
+++ b/Engine/Source/Runtime/NetworkFileSystem/Private/NetworkFileServer.h
@@ -1,9 +1,5 @@
 // Copyright 1998-2014 Epic Games, Inc. All Rights Reserved.
 
-/*=============================================================================
-	StreamingNetworkFileServer.h: Declares the FStreamingNetworkFileServer class.
-=============================================================================*/
-
 #pragma once
 
 /**
@@ -18,8 +14,8 @@ public:
 	/**
 	 * Creates and initializes a new instance.
 	 *
-	 * @param InPort - The port number to bind to (0 = any available port).
-	 * @param InFileRequestDelegate -
+	 * @param InPort The port number to bind to (0 = any available port).
+	 * @param InFileRequestDelegate 
 	 */
 	FNetworkFileServer( int32 InPort, const FFileRequestDelegate* InFileRequestDelegate, 
 		const FRecompileShadersDelegate* InRecompileShadersDelegate, const TArray<ITargetPlatform*>& InActiveTargetPlatforms );
@@ -29,10 +25,9 @@ public:
 	 */
 	~FNetworkFileServer( );
 
-
 public:
 
-	// Begin FRunnable Interface
+	// FRunnable Interface
 
 	virtual bool Init( ) OVERRIDE
 	{
@@ -48,19 +43,15 @@ public:
 
 	virtual void Exit( ) OVERRIDE;
 
-	// End FRunnable Interface
-
-
 public:
 
-	// INetworkFileServer Interface.
+	// INetworkFileServer interface
 
 	virtual bool IsItReadyToAcceptConnections(void) const; 
 	virtual bool GetAddressList(TArray<TSharedPtr<FInternetAddr> >& OutAddresses) const;
 	virtual int32 NumConnections() const;
 	virtual void Shutdown();
 
-
 private:
 
 	// Holds the port to use.
@@ -81,7 +72,6 @@ private:
 	// Is the Listner thread up and running. 
 	FThreadSafeCounter Running;
 
-
 public:
 
 	// Holds a delegate to be invoked on every sync request.
@@ -95,4 +85,4 @@ public:
 
 	// Holds the address that the server is bound to.
 	TSharedPtr<FInternetAddr> ListenAddr;
-};
\ No newline at end of file
+};
diff --git a/Engine/Source/Runtime/NetworkFileSystem/Private/NetworkFileServerConnection.cpp b/Engine/Source/Runtime/NetworkFileSystem/Private/NetworkFileServerConnection.cpp
--- a/Engine/Source/Runtime/NetworkFileSystem/Private/NetworkFileServerConnection.cpp
+++ b/Engine/Source/Runtime/NetworkFileSystem/Private/NetworkFileServerConnection.cpp
@@ -1,9 +1,5 @@
 // Copyright 1998-2014 Epic Games, Inc. All Rights Reserved.
 
-/*=============================================================================
-	NetworkFileServerConnection.h: Declares the FNetworkFileServerClientConnection class.
-=============================================================================*/
-
 #include "NetworkFileSystemPrivatePCH.h"
 #include "PackageName.h"
 #include "TargetPlatform.h"
diff --git a/Engine/Source/Runtime/NetworkFileSystem/Private/NetworkFileServerConnection.h b/Engine/Source/Runtime/NetworkFileSystem/Private/NetworkFileServerConnection.h
--- a/Engine/Source/Runtime/NetworkFileSystem/Private/NetworkFileServerConnection.h
+++ b/Engine/Source/Runtime/NetworkFileSystem/Private/NetworkFileServerConnection.h
@@ -1,11 +1,8 @@
 // Copyright 1998-2014 Epic Games, Inc. All Rights Reserved.
 
-/*=============================================================================
-	NetworkFileServerClientConnection.h: Declares the NetworkFileServerClientConnection class.
-=============================================================================*/
-
 #pragma once
 
+
 /**
  * This class processes all incoming messages from the client.
  */
@@ -28,12 +25,23 @@ public:
 	virtual ~FNetworkFileServerClientConnection( );
 
 public:
-	/** Executes actions from the received payload. */
+
+	/**
+	 * Processes the given payload.
+	 *
+	 * @param Ar An archive containing the payload data.
+	 * @param Out An archive that will contain the processed data.
+	 */
 	void ProcessPayload(FArchive& Ar,FArchive& Out);
 
+	/**
+	 * Gets the client connection's description.
+	 *
+	 * @return Description string.
+	 */
 	FString GetDescription() const;
 
-private:
+protected:
 
 	/**
 	 *	Convert the given filename from the client to the server version of it
@@ -50,6 +58,7 @@ private:
 	 *	@param	FilenameToConvert		Upon input, the server version of the filename. After the call, the client version
 	 */
 	void ConvertServerFilenameToClientFilename(FString& FilenameToConvert);
+	
 	/** Opens a file for reading or writing. */
 	void ProcessOpenFile(FArchive& In, FArchive& Out, bool bIsWriting);
 
@@ -138,6 +147,8 @@ private:
 	 */
 	void ProcessSyncFile( FArchive& In, FArchive& Out );
 
+private:
+
 	// Hold the name of the currently connected platform.
 	FString ConnectedPlatformName;
 
@@ -176,5 +187,4 @@ private:
 
 	// cached copy of the active target platforms (if any)
 	const TArray<ITargetPlatform*>& ActiveTargetPlatforms;
-
-};
\ No newline at end of file
+};
diff --git a/Engine/Source/Runtime/NetworkFileSystem/Private/NetworkFileSystemModule.cpp b/Engine/Source/Runtime/NetworkFileSystem/Private/NetworkFileSystemModule.cpp
--- a/Engine/Source/Runtime/NetworkFileSystem/Private/NetworkFileSystemModule.cpp
+++ b/Engine/Source/Runtime/NetworkFileSystem/Private/NetworkFileSystemModule.cpp
@@ -1,16 +1,12 @@
 // Copyright 1998-2014 Epic Games, Inc. All Rights Reserved.
 
-/*=============================================================================
-	NetworkFileSystemModule.cpp: Implements the FNetworkFileSystemModule class.
-=============================================================================*/
-
 #include "NetworkFileSystemPrivatePCH.h"
-
 #include "TargetDeviceId.h"
 #include "ITargetDevice.h"
 #include "ITargetPlatformManagerModule.h"
 #include "ModuleManager.h"
 
+
 DEFINE_LOG_CATEGORY(LogFileServer);
 
 
@@ -22,7 +18,7 @@ class FNetworkFileSystemModule
 {
 public:
 
-	// Begin INetworkFileSystemModule interface
+	// INetworkFileSystemModule interface
 
 	virtual INetworkFileServer* CreateNetworkFileServer( int32 Port, const FFileRequestDelegate* InFileRequestDelegate, const FRecompileShadersDelegate* InRecompileShadersDelegate ) const OVERRIDE
 	{
@@ -48,9 +44,7 @@ public:
 		return new FNetworkFileServerHttp(Port, InFileRequestDelegate, InRecompileShadersDelegate, ActiveTargetPlatforms);
 #endif 
 	}
-
-	// End INetworkFileSystemModule interface
 };
 
 
-IMPLEMENT_MODULE(FNetworkFileSystemModule, NetworkFileSystem);
\ No newline at end of file
+IMPLEMENT_MODULE(FNetworkFileSystemModule, NetworkFileSystem);
diff --git a/Engine/Source/Runtime/NetworkFileSystem/Private/NetworkFileSystemPrivatePCH.h b/Engine/Source/Runtime/NetworkFileSystem/Private/NetworkFileSystemPrivatePCH.h
--- a/Engine/Source/Runtime/NetworkFileSystem/Private/NetworkFileSystemPrivatePCH.h
+++ b/Engine/Source/Runtime/NetworkFileSystem/Private/NetworkFileSystemPrivatePCH.h
@@ -1,17 +1,11 @@
 // Copyright 1998-2014 Epic Games, Inc. All Rights Reserved.
 
-/*=============================================================================
-	NetworkFileSystemPrivatePCH.h: Pre-compiled header file for the NetworkFileSystem module.
-=============================================================================*/
+#pragma once
 
-#ifndef NETWORK_FILESYSTEM_PRIVATEPCH_H
-#define NETWORK_FILESYSTEM_PRIVATEPCH_H
+#include "NetworkFileSystem.h"
 
 
-#include "../Public/NetworkFileSystem.h"
-
-
-/* Dependencies
+/* Private dependencies
  *****************************************************************************/
 
 #include "Developer/DirectoryWatcher/Public/DirectoryWatcherModule.h"
@@ -30,5 +24,3 @@ DECLARE_LOG_CATEGORY_EXTERN(LogFileServer, Log, All);
 #include "NetworkFileServerConnection.h"
 #include "NetworkFileServer.h"
 #include "NetworkFileServerHttp.h"
-
-#endif
\ No newline at end of file
diff --git a/Engine/Source/Runtime/NetworkFileSystem/Public/Interfaces/INetworkFileServer.h b/Engine/Source/Runtime/NetworkFileSystem/Public/Interfaces/INetworkFileServer.h
--- a/Engine/Source/Runtime/NetworkFileSystem/Public/Interfaces/INetworkFileServer.h
+++ b/Engine/Source/Runtime/NetworkFileSystem/Public/Interfaces/INetworkFileServer.h
@@ -1,9 +1,5 @@
 // Copyright 1998-2014 Epic Games, Inc. All Rights Reserved.
 
-/*=============================================================================
-	INetworkFileServer.h: Declares the INetworkFileServer interface.
-=============================================================================*/
-
 #pragma once
 
 
@@ -20,13 +16,14 @@ public:
 	 * @return true on success, false otherwise.
 	 */
 	 virtual bool IsItReadyToAcceptConnections(void) const = 0;
-	/**
-	 * Gets the list of local network addresses that the file server listens on.
-	 *
-	 * @param OutAddresses - Will hold the address list.
-	 *
-	 * @return true on success, false otherwise.
-	 */
+
+	 /**
+	  * Gets the list of local network addresses that the file server listens on.
+	  *
+	  * @param OutAddresses - Will hold the address list.
+	  *
+	  * @return true on success, false otherwise.
+	  */
 	virtual bool GetAddressList( TArray<TSharedPtr<FInternetAddr> >& OutAddresses ) const = 0;
 
 	/**
@@ -41,6 +38,8 @@ public:
 	 */
 	virtual void Shutdown(void) = 0;
 
+public:
+
 	/**
 	 * Virtual destructor.
 	 */
diff --git a/Engine/Source/Runtime/NetworkFileSystem/Public/Interfaces/INetworkFileSystemModule.h b/Engine/Source/Runtime/NetworkFileSystem/Public/Interfaces/INetworkFileSystemModule.h
--- a/Engine/Source/Runtime/NetworkFileSystem/Public/Interfaces/INetworkFileSystemModule.h
+++ b/Engine/Source/Runtime/NetworkFileSystem/Public/Interfaces/INetworkFileSystemModule.h
@@ -1,9 +1,5 @@
 // Copyright 1998-2014 Epic Games, Inc. All Rights Reserved.
 
-/*=============================================================================
-	INetworkFileSystemModule.h: Declares the INetworkFileSystemModule interface.
-=============================================================================*/
-
 #pragma once
 
 
@@ -42,6 +38,7 @@ struct FShaderRecompileData
 	}
 };
 
+
 /**
  * Delegate type for handling shader recompilation requests from a network client.
  */
@@ -59,15 +56,14 @@ public:
 	/**
 	 * Creates a new network file server.
 	 *
-	 * @param InPort - The port number to bind to (-1 = default port, 0 = any available port).
-	 * @param Streaming - Whether it should be a streaming server.
-	 * @param InFileRequestDelegate - An optional delegate to be invoked when a file is requested by a client.
-	 * @param InRecompileShadersDelegate - An optional delegate to be invoked when shaders need to be recompiled.
+	 * @param InPort The port number to bind to (-1 = default port, 0 = any available port).
+	 * @param Streaming Whether it should be a streaming server.
+	 * @param InFileRequestDelegate An optional delegate to be invoked when a file is requested by a client.
+	 * @param InRecompileShadersDelegate An optional delegate to be invoked when shaders need to be recompiled.
 	 *
-	 * @return The new file server, or NULL if creation failed.
+	 * @return The new file server, or nullptr if creation failed.
 	 */
-	virtual INetworkFileServer* CreateNetworkFileServer( int32 Port = -1, const FFileRequestDelegate* InFileRequestDelegate = NULL, const FRecompileShadersDelegate* InRecompileShadersDelegate = NULL ) const = 0;
-
+	virtual INetworkFileServer* CreateNetworkFileServer( int32 Port = -1, const FFileRequestDelegate* InFileRequestDelegate = nullptr, const FRecompileShadersDelegate* InRecompileShadersDelegate = nullptr ) const = 0;
 
 public:
 
diff --git a/Engine/Source/Runtime/NetworkFileSystem/Public/NetworkFileSystem.h b/Engine/Source/Runtime/NetworkFileSystem/Public/NetworkFileSystem.h
--- a/Engine/Source/Runtime/NetworkFileSystem/Public/NetworkFileSystem.h
+++ b/Engine/Source/Runtime/NetworkFileSystem/Public/NetworkFileSystem.h
@@ -1,13 +1,9 @@
 // Copyright 1998-2014 Epic Games, Inc. All Rights Reserved.
 
-/*=============================================================================
-	NetworkFileSystem.h: NetworkFileSystem module public header file.
-=============================================================================*/
-
 #pragma once
 
 
-/* Dependencies
+/* Public Dependencies
  *****************************************************************************/
 
 #include "Core.h"
@@ -15,12 +11,8 @@
 #include "Sockets.h"
 
 
-/* Interfaces
+/* Public Includes
  *****************************************************************************/
 
 #include "INetworkFileServer.h"
 #include "INetworkFileSystemModule.h"
-
-
-/* Common
- *****************************************************************************/
\ No newline at end of file
diff --git a/Engine/Source/Runtime/Networking/Private/IPv4/IPv4Address.cpp b/Engine/Source/Runtime/Networking/Private/IPv4/IPv4Address.cpp
--- a/Engine/Source/Runtime/Networking/Private/IPv4/IPv4Address.cpp
+++ b/Engine/Source/Runtime/Networking/Private/IPv4/IPv4Address.cpp
@@ -1,9 +1,5 @@
 // Copyright 1998-2014 Epic Games, Inc. All Rights Reserved.
 
-/*=============================================================================
-	IPv4Address.cpp: Implements the FIPv4Address class.
-=============================================================================*/
-
 #include "NetworkingPrivatePCH.h"
 
 
diff --git a/Engine/Source/Runtime/Networking/Private/IPv4/IPv4Endpoint.cpp b/Engine/Source/Runtime/Networking/Private/IPv4/IPv4Endpoint.cpp
--- a/Engine/Source/Runtime/Networking/Private/IPv4/IPv4Endpoint.cpp
+++ b/Engine/Source/Runtime/Networking/Private/IPv4/IPv4Endpoint.cpp
@@ -1,9 +1,5 @@
 // Copyright 1998-2014 Epic Games, Inc. All Rights Reserved.
 
-/*=============================================================================
-	IPv4Endpoint.cpp: Implements the FIPv4Endpoint class.
-=============================================================================*/
-
 #include "NetworkingPrivatePCH.h"
 
 
@@ -11,7 +7,7 @@
  *****************************************************************************/
 
 const FIPv4Endpoint FIPv4Endpoint::Any(FIPv4Address(0), 0);
-ISocketSubsystem* FIPv4Endpoint::CachedSocketSubsystem = NULL;
+ISocketSubsystem* FIPv4Endpoint::CachedSocketSubsystem = nullptr;
 
 
 /* FIPv4Endpoint interface
diff --git a/Engine/Source/Runtime/Networking/Private/IPv4/IPv4Subnet.cpp b/Engine/Source/Runtime/Networking/Private/IPv4/IPv4Subnet.cpp
--- a/Engine/Source/Runtime/Networking/Private/IPv4/IPv4Subnet.cpp
+++ b/Engine/Source/Runtime/Networking/Private/IPv4/IPv4Subnet.cpp
@@ -1,9 +1,5 @@
 // Copyright 1998-2014 Epic Games, Inc. All Rights Reserved.
 
-/*=============================================================================
-	IPv4Subnet.h: Implements the FIPv4Subnet class.
-=============================================================================*/
-
 #include "NetworkingPrivatePCH.h"
 
 
diff --git a/Engine/Source/Runtime/Networking/Private/IPv4/IPv4SubnetMask.cpp b/Engine/Source/Runtime/Networking/Private/IPv4/IPv4SubnetMask.cpp
--- a/Engine/Source/Runtime/Networking/Private/IPv4/IPv4SubnetMask.cpp
+++ b/Engine/Source/Runtime/Networking/Private/IPv4/IPv4SubnetMask.cpp
@@ -1,9 +1,5 @@
 // Copyright 1998-2014 Epic Games, Inc. All Rights Reserved.
 
-/*=============================================================================
-	IPv4SubnetMask.cpp: Implements the FIPv4SubnetMask class.
-=============================================================================*/
-
 #include "NetworkingPrivatePCH.h"
 
 
diff --git a/Engine/Source/Runtime/Networking/Private/NetworkingModule.cpp b/Engine/Source/Runtime/Networking/Private/NetworkingModule.cpp
--- a/Engine/Source/Runtime/Networking/Private/NetworkingModule.cpp
+++ b/Engine/Source/Runtime/Networking/Private/NetworkingModule.cpp
@@ -1,11 +1,6 @@
 // Copyright 1998-2014 Epic Games, Inc. All Rights Reserved.
 
-/*=============================================================================
-	NetworkingModule.cpp: Implements the FNetworkingModule class.
-=============================================================================*/
-
 #include "NetworkingPrivatePCH.h"
-
 #include "ModuleManager.h"
 
 
@@ -29,4 +24,4 @@ void FNetworkingModule::ShutdownModule( )
 bool FNetworkingModule::SupportsDynamicReloading( )
 {
 	return false;
-}
\ No newline at end of file
+}
diff --git a/Engine/Source/Runtime/Networking/Private/NetworkingModule.h b/Engine/Source/Runtime/Networking/Private/NetworkingModule.h
--- a/Engine/Source/Runtime/Networking/Private/NetworkingModule.h
+++ b/Engine/Source/Runtime/Networking/Private/NetworkingModule.h
@@ -1,9 +1,5 @@
 // Copyright 1998-2014 Epic Games, Inc. All Rights Reserved.
 
-/*=============================================================================
-	NetworkingModule.h: Declares the FNetworkingModule class.
-=============================================================================*/
-
 #pragma once
 
 
diff --git a/Engine/Source/Runtime/Networking/Private/NetworkingPrivatePCH.h b/Engine/Source/Runtime/Networking/Private/NetworkingPrivatePCH.h
--- a/Engine/Source/Runtime/Networking/Private/NetworkingPrivatePCH.h
+++ b/Engine/Source/Runtime/Networking/Private/NetworkingPrivatePCH.h
@@ -1,17 +1,13 @@
 // Copyright 1998-2014 Epic Games, Inc. All Rights Reserved.
 
-/*=============================================================================
-	NetworkingPrivatePCH.h: Pre-compiled header file for the Networking module.
-=============================================================================*/
+#pragma once
+
+#include "Networking.h"
 
-#ifndef UNREAL_NETWORKING_PRIVATE_PCH
-#define UNREAL_NETWORKING_PRIVATE_PCH
 
 /* Dependencies
  *****************************************************************************/
 
-#include "../Public/Networking.h"
-
 #include "Sockets.h"
 #include "SocketSubsystem.h"
 
@@ -20,5 +16,3 @@
  *****************************************************************************/
 
 #include "NetworkingModule.h"
-
-#endif
diff --git a/Engine/Source/Runtime/Networking/Private/Steam/SteamEndpoint.cpp b/Engine/Source/Runtime/Networking/Private/Steam/SteamEndpoint.cpp
--- a/Engine/Source/Runtime/Networking/Private/Steam/SteamEndpoint.cpp
+++ b/Engine/Source/Runtime/Networking/Private/Steam/SteamEndpoint.cpp
@@ -1,9 +1,5 @@
 // Copyright 1998-2014 Epic Games, Inc. All Rights Reserved.
 
-/*=============================================================================
-	SteamEndpoint.h: Implements the FSteamEndpoint class.
-=============================================================================*/
-
 #include "NetworkingPrivatePCH.h"
 
 
@@ -13,4 +9,4 @@
 FString FSteamEndpoint::ToString( ) const
 {
 	return FString::Printf(TEXT("0x%llX:%i"), UniqueNetId, SteamChannel);
-}
\ No newline at end of file
+}
diff --git a/Engine/Source/Runtime/Networking/Private/Tests/IPv4AddressTest.cpp b/Engine/Source/Runtime/Networking/Private/Tests/IPv4AddressTest.cpp
--- a/Engine/Source/Runtime/Networking/Private/Tests/IPv4AddressTest.cpp
+++ b/Engine/Source/Runtime/Networking/Private/Tests/IPv4AddressTest.cpp
@@ -1,11 +1,6 @@
 // Copyright 1998-2014 Epic Games, Inc. All Rights Reserved.
 
-/*=============================================================================
-	IPv4AddresTest.cpp: Unit test for the FIPv4AddresTest class.
-=============================================================================*/
-
 #include "NetworkingPrivatePCH.h"
-
 #include "AutomationTest.h"
 
 
@@ -88,4 +83,4 @@ bool FIPv4AddressTest::RunTest( const FString& Parameters )
 	TestFalse(TEXT("192.169.0.1 must not be a site local address"), a8_8.IsSiteLocalAddress());
 
 	return true;
-}
\ No newline at end of file
+}
diff --git a/Engine/Source/Runtime/Networking/Public/Common/TcpListener.h b/Engine/Source/Runtime/Networking/Public/Common/TcpListener.h
--- a/Engine/Source/Runtime/Networking/Public/Common/TcpListener.h
+++ b/Engine/Source/Runtime/Networking/Public/Common/TcpListener.h
@@ -1,9 +1,5 @@
 // Copyright 1998-2014 Epic Games, Inc. All Rights Reserved.
 
-/*=============================================================================
-	TcpListener.h: Declares the FTcpListener class.
-=============================================================================*/
-
 #pragma once
 
 
@@ -28,8 +24,8 @@ public:
 	/**
 	 * Creates and initializes a new instance from the specified IP endpoint.
 	 *
-	 * @param LocalEndpoint - The local IP endpoint to listen on.
-	 * @param SleepTime - The time to sleep between checking for pending connections (default = 0 seconds).
+	 * @param LocalEndpoint The local IP endpoint to listen on.
+	 * @param SleepTime The time to sleep between checking for pending connections (default = 0 seconds).
 	 */
 	FTcpListener( const FIPv4Endpoint& LocalEndpoint, const FTimespan& InSleepTime = FTimespan::Zero())
 		: DeleteSocket(true)
@@ -44,8 +40,8 @@ public:
 	/**
 	 * Creates and initializes a new instance from the specified socket.
 	 *
-	 * @param InSocket - The socket to listen on.
-	 * @param SleepTime - The time to sleep between checking for pending connections (default = 0 seconds).
+	 * @param InSocket The socket to listen on.
+	 * @param SleepTime The time to sleep between checking for pending connections (default = 0 seconds).
 	 */
 	FTcpListener( FSocket& InSocket, const FTimespan& InSleepTime = FTimespan::Zero() )
 		: DeleteSocket(false)
@@ -78,7 +74,6 @@ public:
 		}
 	}
 
-
 public:
 
 	/**
@@ -111,7 +106,6 @@ public:
 		return ((Socket != NULL) && !Stopping);
 	}
 
-
 public:
 
 	/**
@@ -121,19 +115,16 @@ public:
 	 * To temporarily disable accepting connections, use the Enable() and Disable() methods.
 	 *
 	 * @return The delegate.
-	 *
-	 * @see Enable
-	 * @see Disable
+	 * @see Enable, Disable
 	 */
 	FOnTcpListenerConnectionAccepted& OnConnectionAccepted( )
 	{
 		return ConnectionAcceptedDelegate;
 	}
 
-
 public:
 
-	// Begin FRunnable interface
+	// FRunnable interface
 
 	virtual bool Init( ) override
 	{
@@ -194,9 +185,6 @@ public:
 
 	virtual void Exit( ) override { }
 
-	// End FRunnable interface
-
-
 private:
 
 	// Holds a flag indicating whether the socket should be deleted in the destructor.
@@ -217,7 +205,6 @@ private:
 	// Holds the thread object.
 	FRunnableThread* Thread;
 
-
 private:
 
 	// Holds a delegate to be invoked when an incoming connection has been accepted.
diff --git a/Engine/Source/Runtime/Networking/Public/Common/TcpSocketBuilder.h b/Engine/Source/Runtime/Networking/Public/Common/TcpSocketBuilder.h
--- a/Engine/Source/Runtime/Networking/Public/Common/TcpSocketBuilder.h
+++ b/Engine/Source/Runtime/Networking/Public/Common/TcpSocketBuilder.h
@@ -1,9 +1,5 @@
 // Copyright 1998-2014 Epic Games, Inc. All Rights Reserved.
 
-/*=============================================================================
-	TcpSocketBuilder.h: Declares the FTcpSocketBuilder class.
-=============================================================================*/
-
 #pragma once
 
 
diff --git a/Engine/Source/Runtime/Networking/Public/Common/UdpSocketBuilder.h b/Engine/Source/Runtime/Networking/Public/Common/UdpSocketBuilder.h
--- a/Engine/Source/Runtime/Networking/Public/Common/UdpSocketBuilder.h
+++ b/Engine/Source/Runtime/Networking/Public/Common/UdpSocketBuilder.h
@@ -1,9 +1,5 @@
 // Copyright 1998-2014 Epic Games, Inc. All Rights Reserved.
 
-/*=============================================================================
-	UdpSocketBuilder.h: Declares the FUdpSocketBuilder class.
-=============================================================================*/
-
 #pragma once
 
 
diff --git a/Engine/Source/Runtime/Networking/Public/Common/UdpSocketReceiver.h b/Engine/Source/Runtime/Networking/Public/Common/UdpSocketReceiver.h
--- a/Engine/Source/Runtime/Networking/Public/Common/UdpSocketReceiver.h
+++ b/Engine/Source/Runtime/Networking/Public/Common/UdpSocketReceiver.h
@@ -1,9 +1,5 @@
 // Copyright 1998-2014 Epic Games, Inc. All Rights Reserved.
 
-/*=============================================================================
-	UdpSocketReceiver.h: Declares the FUdpSocketReceiver class.
-=============================================================================*/
-
 #pragma once
 
 
diff --git a/Engine/Source/Runtime/Networking/Public/Common/UdpSocketSender.h b/Engine/Source/Runtime/Networking/Public/Common/UdpSocketSender.h
--- a/Engine/Source/Runtime/Networking/Public/Common/UdpSocketSender.h
+++ b/Engine/Source/Runtime/Networking/Public/Common/UdpSocketSender.h
@@ -1,9 +1,5 @@
 // Copyright 1998-2014 Epic Games, Inc. All Rights Reserved.
 
-/*=============================================================================
-	UdpSocketSender.h: Declares the FUdpSocketSender class.
-=============================================================================*/
-
 #pragma once
 
 
diff --git a/Engine/Source/Runtime/Networking/Public/Interfaces/INetworkingModule.h b/Engine/Source/Runtime/Networking/Public/Interfaces/INetworkingModule.h
--- a/Engine/Source/Runtime/Networking/Public/Interfaces/INetworkingModule.h
+++ b/Engine/Source/Runtime/Networking/Public/Interfaces/INetworkingModule.h
@@ -1,9 +1,5 @@
 // Copyright 1998-2014 Epic Games, Inc. All Rights Reserved.
 
-/*=============================================================================
-	INetworkingModule.h: Declares the INetworkingModule interface.
-=============================================================================*/
-
 #pragma once
 
 
diff --git a/Engine/Source/Runtime/Networking/Public/Interfaces/IPv4/IPv4Address.h b/Engine/Source/Runtime/Networking/Public/Interfaces/IPv4/IPv4Address.h
--- a/Engine/Source/Runtime/Networking/Public/Interfaces/IPv4/IPv4Address.h
+++ b/Engine/Source/Runtime/Networking/Public/Interfaces/IPv4/IPv4Address.h
@@ -1,9 +1,5 @@
 // Copyright 1998-2014 Epic Games, Inc. All Rights Reserved.
 
-/*=============================================================================
-	IPv4Address.h: Declares the FIPv4Address class.
-=============================================================================*/
-
 #pragma once
 
 
diff --git a/Engine/Source/Runtime/Networking/Public/Interfaces/IPv4/IPv4Endpoint.h b/Engine/Source/Runtime/Networking/Public/Interfaces/IPv4/IPv4Endpoint.h
--- a/Engine/Source/Runtime/Networking/Public/Interfaces/IPv4/IPv4Endpoint.h
+++ b/Engine/Source/Runtime/Networking/Public/Interfaces/IPv4/IPv4Endpoint.h
@@ -1,9 +1,5 @@
 // Copyright 1998-2014 Epic Games, Inc. All Rights Reserved.
 
-/*=============================================================================
-	IPv4Endpoint.h: Declares the FIPv4Endpoint class.
-=============================================================================*/
-
 #pragma once
 
 
diff --git a/Engine/Source/Runtime/Networking/Public/Interfaces/IPv4/IPv4Subnet.h b/Engine/Source/Runtime/Networking/Public/Interfaces/IPv4/IPv4Subnet.h
--- a/Engine/Source/Runtime/Networking/Public/Interfaces/IPv4/IPv4Subnet.h
+++ b/Engine/Source/Runtime/Networking/Public/Interfaces/IPv4/IPv4Subnet.h
@@ -1,9 +1,5 @@
 // Copyright 1998-2014 Epic Games, Inc. All Rights Reserved.
 
-/*=============================================================================
-	IPv4Subnet.h: Declares the FIPv4Subnet class.
-=============================================================================*/
-
 #pragma once
 
 
diff --git a/Engine/Source/Runtime/Networking/Public/Interfaces/IPv4/IPv4SubnetMask.h b/Engine/Source/Runtime/Networking/Public/Interfaces/IPv4/IPv4SubnetMask.h
--- a/Engine/Source/Runtime/Networking/Public/Interfaces/IPv4/IPv4SubnetMask.h
+++ b/Engine/Source/Runtime/Networking/Public/Interfaces/IPv4/IPv4SubnetMask.h
@@ -1,9 +1,5 @@
 // Copyright 1998-2014 Epic Games, Inc. All Rights Reserved.
 
-/*=============================================================================
-	IPv4SubnetMask.h: Declares the FIPv4SubnetMask class.
-=============================================================================*/
-
 #pragma once
 
 
diff --git a/Engine/Source/Runtime/Networking/Public/Interfaces/Steam/SteamEndpoint.h b/Engine/Source/Runtime/Networking/Public/Interfaces/Steam/SteamEndpoint.h
--- a/Engine/Source/Runtime/Networking/Public/Interfaces/Steam/SteamEndpoint.h
+++ b/Engine/Source/Runtime/Networking/Public/Interfaces/Steam/SteamEndpoint.h
@@ -1,9 +1,5 @@
 // Copyright 1998-2014 Epic Games, Inc. All Rights Reserved.
 
-/*=============================================================================
-	SteamEndpoint.h: Declares the FSteamEndpoint class.
-=============================================================================*/
-
 #pragma once
 
 
diff --git a/Engine/Source/Runtime/Networking/Public/Networking.h b/Engine/Source/Runtime/Networking/Public/Networking.h
--- a/Engine/Source/Runtime/Networking/Public/Networking.h
+++ b/Engine/Source/Runtime/Networking/Public/Networking.h
@@ -1,13 +1,9 @@
 // Copyright 1998-2014 Epic Games, Inc. All Rights Reserved.
 
-/*=============================================================================
-	Networking.h: Networking module public header file.
-=============================================================================*/
-
 #pragma once
 
 
-/* Dependencies
+/* Public Dependencies
  *****************************************************************************/
 
 #include "Core.h"
@@ -16,16 +12,14 @@
 #include "SocketSubsystem.h"
 
 
-/* Interfaces
+/* Public Includes
  *****************************************************************************/
 
 #include "IPv4SubnetMask.h"
 #include "IPv4Address.h"
 #include "IPv4Endpoint.h"
 #include "IPv4Subnet.h"
-
 #include "SteamEndpoint.h"
-
 #include "INetworkingModule.h"
 
 
@@ -34,7 +28,6 @@
 
 #include "TcpSocketBuilder.h"
 #include "TcpListener.h"
-
 #include "UdpSocketReceiver.h"
 #include "UdpSocketSender.h"
 #include "UdpSocketBuilder.h"
\ No newline at end of file
diff --git a/Engine/Source/Runtime/Projects/Public/Projects.h b/Engine/Source/Runtime/Projects/Public/Projects.h
--- a/Engine/Source/Runtime/Projects/Public/Projects.h
+++ b/Engine/Source/Runtime/Projects/Public/Projects.h
@@ -1,19 +1,15 @@
 // Copyright 1998-2014 Epic Games, Inc. All Rights Reserved.
 
-/*=============================================================================
-	Projects.h: Projects module public header file.
-=============================================================================*/
-
 #pragma once
 
 
-/* Dependencies
+/* Public Dependencies
  *****************************************************************************/
 
 #include "Core.h"
 
 
-/* Interfaces
+/* Public Includes
  *****************************************************************************/
 
 #include "IPluginManagerShared.h"
diff --git a/Engine/Source/Runtime/SlateRHIRenderer/Private/SlateRHIFontTexture.cpp b/Engine/Source/Runtime/SlateRHIRenderer/Private/SlateRHIFontTexture.cpp
--- a/Engine/Source/Runtime/SlateRHIRenderer/Private/SlateRHIFontTexture.cpp
+++ b/Engine/Source/Runtime/SlateRHIRenderer/Private/SlateRHIFontTexture.cpp
@@ -1,9 +1,5 @@
 // Copyright 1998-2014 Epic Games, Inc. All Rights Reserved.
 
-/*=============================================================================
-	SlateRHIFontTexture.cpp: Implements the FSlateFontTextureRHI class.
-=============================================================================*/
-
 #include "SlateRHIRendererPrivatePCH.h"
 
 
diff --git a/Engine/Source/Runtime/SlateRHIRenderer/Private/SlateRHIRenderer.cpp b/Engine/Source/Runtime/SlateRHIRenderer/Private/SlateRHIRenderer.cpp
--- a/Engine/Source/Runtime/SlateRHIRenderer/Private/SlateRHIRenderer.cpp
+++ b/Engine/Source/Runtime/SlateRHIRenderer/Private/SlateRHIRenderer.cpp
@@ -1,11 +1,8 @@
 // Copyright 1998-2014 Epic Games, Inc. All Rights Reserved.
 
-/*=============================================================================
-	SlateRHIRenderer.cpp: Implements the FSlateRenderer class and related classes.
-=============================================================================*/
-
 #include "SlateRHIRendererPrivatePCH.h"
 
+
 DECLARE_CYCLE_STAT(TEXT("Map Staging Buffer"),STAT_MapStagingBuffer,STATGROUP_CrashTracker);
 DECLARE_CYCLE_STAT(TEXT("Generate Capture Buffer"),STAT_GenerateCaptureBuffer,STATGROUP_CrashTracker);
 DECLARE_CYCLE_STAT(TEXT("Unmap Staging Buffer"),STAT_UnmapStagingBuffer,STATGROUP_CrashTracker);
diff --git a/Engine/Source/Runtime/SlateRHIRenderer/Private/SlateRHIRendererModule.cpp b/Engine/Source/Runtime/SlateRHIRenderer/Private/SlateRHIRendererModule.cpp
--- a/Engine/Source/Runtime/SlateRHIRenderer/Private/SlateRHIRendererModule.cpp
+++ b/Engine/Source/Runtime/SlateRHIRenderer/Private/SlateRHIRendererModule.cpp
@@ -1,9 +1,5 @@
 // Copyright 1998-2014 Epic Games, Inc. All Rights Reserved.
 
-/*=============================================================================
-	SlateRHIRendererModule.cpp: Implements the FSlateRHIRendererModule class.
-=============================================================================*/
-
 #include "SlateRHIRendererPrivatePCH.h"
 
 
@@ -15,7 +11,7 @@ class FSlateRHIRendererModule
 {
 public:
 
-	// Begin ISlateRHIRendererModule interface
+	// ISlateRHIRendererModule interface
 
 	virtual TSharedRef<FSlateRenderer> CreateSlateRHIRenderer( ) override
 	{
@@ -23,10 +19,7 @@ public:
 	}
 
 	virtual void StartupModule( ) override { }
-	
 	virtual void ShutdownModule( ) override { }
-
-	// End ISlateRHIRendererModule interface
 };
 
 
diff --git a/Engine/Source/Runtime/SlateRHIRenderer/Private/SlateRHIRendererPrivatePCH.h b/Engine/Source/Runtime/SlateRHIRenderer/Private/SlateRHIRendererPrivatePCH.h
--- a/Engine/Source/Runtime/SlateRHIRenderer/Private/SlateRHIRendererPrivatePCH.h
+++ b/Engine/Source/Runtime/SlateRHIRenderer/Private/SlateRHIRendererPrivatePCH.h
@@ -1,15 +1,11 @@
 // Copyright 1998-2014 Epic Games, Inc. All Rights Reserved.
 
-/*=============================================================================
-	SlateRHIRendererPrivatePCH.h: Pre-compiled header file for the SlateRHIRenderer module.
-=============================================================================*/
-
 #pragma once
 
 #include "ISlateRHIRendererModule.h"
 
 
-/* Private Dependencies
+/* Private dependencies
  *****************************************************************************/
 
 #include "Slate.h"
@@ -24,7 +20,7 @@
 #include "ShaderParameterUtils.h"
 
 
-/* Private Includes
+/* Private includes
  *****************************************************************************/
 
 #include "SlateRHIFontTexture.h"
diff --git a/Engine/Source/Runtime/SlateRHIRenderer/Private/SlateRHIRenderingPolicy.cpp b/Engine/Source/Runtime/SlateRHIRenderer/Private/SlateRHIRenderingPolicy.cpp
--- a/Engine/Source/Runtime/SlateRHIRenderer/Private/SlateRHIRenderingPolicy.cpp
+++ b/Engine/Source/Runtime/SlateRHIRenderer/Private/SlateRHIRenderingPolicy.cpp
@@ -1,9 +1,5 @@
 // Copyright 1998-2014 Epic Games, Inc. All Rights Reserved.
 
-/*=============================================================================
-	SlateRHIRenderingPolicy.cpp: Implements the FSlateRHIRenderingPolicy class and related classes.
-=============================================================================*/
-
 #include "SlateRHIRendererPrivatePCH.h"
 #include "RenderingPolicy.h"
 #include "RHIStaticStates.h"
diff --git a/Engine/Source/Runtime/SlateRHIRenderer/Private/SlateRHIResourceManager.cpp b/Engine/Source/Runtime/SlateRHIRenderer/Private/SlateRHIResourceManager.cpp
--- a/Engine/Source/Runtime/SlateRHIRenderer/Private/SlateRHIResourceManager.cpp
+++ b/Engine/Source/Runtime/SlateRHIRenderer/Private/SlateRHIResourceManager.cpp
@@ -1,9 +1,5 @@
 // Copyright 1998-2014 Epic Games, Inc. All Rights Reserved.
 
-/*=============================================================================
-	SlateRHIRenderingPolicy.cpp: Implements the FSlateRHIResourceManager class.
-=============================================================================*/
-
 #include "SlateRHIRendererPrivatePCH.h"
 #include "ImageWrapper.h"
 
diff --git a/Engine/Source/Runtime/SlateRHIRenderer/Private/SlateRHITextureAtlas.cpp b/Engine/Source/Runtime/SlateRHIRenderer/Private/SlateRHITextureAtlas.cpp
--- a/Engine/Source/Runtime/SlateRHIRenderer/Private/SlateRHITextureAtlas.cpp
+++ b/Engine/Source/Runtime/SlateRHIRenderer/Private/SlateRHITextureAtlas.cpp
@@ -1,9 +1,5 @@
 // Copyright 1998-2014 Epic Games, Inc. All Rights Reserved.
 
-/*=============================================================================
-	SlateRHITextureAtlas.cpp: Implements the FSlateTextureAtlasRHI class.
-=============================================================================*/
-
 #include "SlateRHIRendererPrivatePCH.h"
 
 
diff --git a/Engine/Source/Runtime/SlateRHIRenderer/Private/SlateRHITextureAtlas.h b/Engine/Source/Runtime/SlateRHIRenderer/Private/SlateRHITextureAtlas.h
--- a/Engine/Source/Runtime/SlateRHIRenderer/Private/SlateRHITextureAtlas.h
+++ b/Engine/Source/Runtime/SlateRHIRenderer/Private/SlateRHITextureAtlas.h
@@ -1,9 +1,5 @@
 // Copyright 1998-2014 Epic Games, Inc. All Rights Reserved.
 
-/*=============================================================================
-	SlateRHITextureAtlas.h: Declares the FSlateTextureAtlasRHI class.
-=============================================================================*/
-
 #pragma once
 
 
@@ -56,12 +52,10 @@ public:
 
 public:
 
-	// Begin FSlateTextureAtlas overrides.
+	// FSlateTextureAtlas overrides.
 
 	virtual void ConditionalUpdateTexture( );
 
-	// End FSlateTextureAtlas interface.
-
 private:
 
 	/** The texture rendering resource */
diff --git a/Engine/Source/Runtime/SlateRHIRenderer/Private/SlateShaders.cpp b/Engine/Source/Runtime/SlateRHIRenderer/Private/SlateShaders.cpp
--- a/Engine/Source/Runtime/SlateRHIRenderer/Private/SlateShaders.cpp
+++ b/Engine/Source/Runtime/SlateRHIRenderer/Private/SlateShaders.cpp
@@ -1,9 +1,5 @@
 // Copyright 1998-2014 Epic Games, Inc. All Rights Reserved.
 
-/*=============================================================================
-	SlateShaders.cpp: Implements shader related classes for the Slate RHI Renderer.
-=============================================================================*/
-
 #include "SlateRHIRendererPrivatePCH.h"
 
 
diff --git a/Engine/Source/Runtime/SlateRHIRenderer/Public/Interfaces/ISlateRHIRendererModule.h b/Engine/Source/Runtime/SlateRHIRenderer/Public/Interfaces/ISlateRHIRendererModule.h
--- a/Engine/Source/Runtime/SlateRHIRenderer/Public/Interfaces/ISlateRHIRendererModule.h
+++ b/Engine/Source/Runtime/SlateRHIRenderer/Public/Interfaces/ISlateRHIRendererModule.h
@@ -1,9 +1,5 @@
 // Copyright 1998-2014 Epic Games, Inc. All Rights Reserved.
 
-/*=============================================================================
-	ISlateRHIRendererModule.h: Declares the ISlateRHIRendererModule interface.
-=============================================================================*/
-
 #pragma once
 
 #include "ModuleInterface.h"
diff --git a/Engine/Source/Runtime/UMG/Private/Slate/SConstraintCanvas.cpp b/Engine/Source/Runtime/UMG/Private/Slate/SConstraintCanvas.cpp
--- a/Engine/Source/Runtime/UMG/Private/Slate/SConstraintCanvas.cpp
+++ b/Engine/Source/Runtime/UMG/Private/Slate/SConstraintCanvas.cpp
@@ -1,11 +1,6 @@
 // Copyright 1998-2014 Epic Games, Inc. All Rights Reserved.
 
-/*=============================================================================
-	SConstraintCanvas.cpp: Implements the SConstraintCanvas class.
-=============================================================================*/
-
 #include "UMGPrivatePCH.h"
-
 #include "LayoutUtils.h"
 
 
diff --git a/Engine/Source/Runtime/UMG/Public/Slate/SConstraintCanvas.h b/Engine/Source/Runtime/UMG/Public/Slate/SConstraintCanvas.h
--- a/Engine/Source/Runtime/UMG/Public/Slate/SConstraintCanvas.h
+++ b/Engine/Source/Runtime/UMG/Public/Slate/SConstraintCanvas.h
@@ -1,11 +1,8 @@
 // Copyright 1998-2014 Epic Games, Inc. All Rights Reserved.
 
-/*=============================================================================
-	SConstraintCanvas.h: Declares the SConstraintCanvas class.
-=============================================================================*/
-
 #pragma once
 
+
 /**
  * ConstraintCanvas is a layout widget that allows you to arbitrary position and size child widgets in a relative coordinate space
  */
diff --git a/Engine/Source/Runtime/UMG/Public/UMG.h b/Engine/Source/Runtime/UMG/Public/UMG.h
--- a/Engine/Source/Runtime/UMG/Public/UMG.h
+++ b/Engine/Source/Runtime/UMG/Public/UMG.h
@@ -1,13 +1,9 @@
 // Copyright 1998-2014 Epic Games, Inc. All Rights Reserved.
 
-/*=============================================================================
-	UMG.h: UMG module public header file.
-=============================================================================*/
-
 #pragma once
 
 
-/* Dependencies
+/* Public Dependencies
  *****************************************************************************/
 
 #include "Core.h"
@@ -20,9 +16,11 @@
 #include "Anchors.h"
 #include "SConstraintCanvas.h"
 
-/* Components
+
+/* Public Includes
  *****************************************************************************/
 
+// Components
 #include "SlateWrapperTypes.h"
 
 #include "PanelSlot.h"
@@ -70,14 +68,10 @@
 #include "ListView.h"
 #include "TileView.h"
 
-/* Slate
- *****************************************************************************/
- 
+// Slate
 #include "SObjectWidget.h"
 
-/* Blueprint
- *****************************************************************************/
-
+// Blueprint
 #include "WidgetTree.h"
 #include "UserWidget.h"
 
@@ -89,7 +83,5 @@
 
 #include "WidgetActor.h"
 
-/* Interfaces
- *****************************************************************************/
-
+// Interfaces
 #include "IUMGModule.h"
