Renamed NULL to nullptr

diff --git a/Engine/Source/Runtime/Networking/Public/Common/UdpSocketBuilder.h b/Engine/Source/Runtime/Networking/Public/Common/UdpSocketBuilder.h
--- a/Engine/Source/Runtime/Networking/Public/Common/UdpSocketBuilder.h
+++ b/Engine/Source/Runtime/Networking/Public/Common/UdpSocketBuilder.h
@@ -199,15 +199,15 @@ public:
 	 */
 	FSocket* Build( ) const
 	{
-		FSocket* Socket = NULL;
+		FSocket* Socket = nullptr;
 
 		ISocketSubsystem* SocketSubsystem = ISocketSubsystem::Get(PLATFORM_SOCKETSUBSYSTEM);
 
-		if (SocketSubsystem != NULL)
+		if (SocketSubsystem != nullptr)
 		{
 			Socket = SocketSubsystem->CreateSocket(NAME_DGram, *Description, true);
 
-			if (Socket != NULL)
+			if (Socket != nullptr)
 			{
 				bool Error = !Socket->SetNonBlocking(!Blocking) ||
 							 !Socket->SetReuseAddr(Reusable) ||
@@ -243,8 +243,7 @@ public:
 					GLog->Logf(TEXT("FUdpSocketBuilder: Failed to create the socket %s as configured"), *Description);
 
 					SocketSubsystem->DestroySocket(Socket);
-
-					Socket = NULL;
+					Socket = nullptr;
 				}
 			}
 		}
