Sockets: Added API for specifying send and receive buffer sizes on UDP/TCP socket builders

diff --git a/Engine/Source/Runtime/Networking/Public/Common/TcpSocketBuilder.h b/Engine/Source/Runtime/Networking/Public/Common/TcpSocketBuilder.h
--- a/Engine/Source/Runtime/Networking/Public/Common/TcpSocketBuilder.h
+++ b/Engine/Source/Runtime/Networking/Public/Common/TcpSocketBuilder.h
@@ -13,7 +13,7 @@ public:
 	/**
 	 * Creates and initializes a new instance.
 	 *
-	 * @param InDescription - Debug description for the socket.
+	 * @param InDescription Debug description for the socket.
 	 */
 	FTcpSocketBuilder( const FString& InDescription )
 		: Blocking(false)
@@ -23,18 +23,20 @@ public:
 		, Linger(false)
 		, LingerTimeout(0)
 		, Listen(false)
+		, ReceiveBufferSize(0)
 		, Reusable(false)
+		, SendBufferSize(0)
 	{ }
 
-
 public:
 
 	/**
 	 * Sets socket operations to be blocking.
 	 *
 	 * @return This instance (for method chaining).
+	 * @see AsNonBlocking, AsReusable
 	 */
-	FTcpSocketBuilder AsBlocking( )
+	FTcpSocketBuilder AsBlocking()
 	{
 		Blocking = true;
 
@@ -45,8 +47,9 @@ public:
 	 * Sets socket operations to be non-blocking.
 	 *
 	 * @return This instance (for method chaining).
+	 * @see AsBlocking, AsReusable
 	 */
-	FTcpSocketBuilder AsNonBlocking( )
+	FTcpSocketBuilder AsNonBlocking()
 	{
 		Blocking = false;
 
@@ -57,8 +60,9 @@ public:
 	 * Makes the bound address reusable by other sockets.
 	 *
 	 * @return This instance (for method chaining).
+	 * @see AsNonBlocking, AsNonBlocking
 	 */
-	FTcpSocketBuilder AsReusable( )
+	FTcpSocketBuilder AsReusable()
 	{
 		Reusable = true;
 
@@ -71,12 +75,9 @@ public:
 	 * Unless specified in a subsequent call to BoundToPort(), a random
 	 * port number will be assigned by the underlying provider.
 	 *
-	 * @param Address - The IP address to bind the socket to.
-	 *
+	 * @param Address The IP address to bind the socket to.
 	 * @return This instance (for method chaining).
-	 *
-	 * @see BoundToEndpoint
-	 * @see BoundToPort
+	 * @see BoundToEndpoint, BoundToPort
 	 */
 	FTcpSocketBuilder BoundToAddress( const FIPv4Address& Address )
 	{
@@ -89,12 +90,9 @@ public:
 	/**
  	 * Sets the local endpoint to bind the socket to.
 	 *
-	 * @param Endpoint - The IP endpoint to bind the socket to.
-	 *
+	 * @param Endpoint The IP endpoint to bind the socket to.
 	 * @return This instance (for method chaining).
-	 *
-	 * @see BoundToAddress
-	 * @see BoundToPort
+	 * @see BoundToAddress, BoundToPort
 	 */
 	FTcpSocketBuilder BoundToEndpoint( const FIPv4Endpoint& Endpoint )
 	{
@@ -110,11 +108,9 @@ public:
 	 * Unless specified in a subsequent call to BoundToAddress(), the local
 	 * address will be determined automatically by the underlying provider.
 	 *
-	 * @param Port - The local port number to bind the socket to.
-	 *
+	 * @param Port The local port number to bind the socket to.
 	 * @return This instance (for method chaining).
-	 *
-	 * @see BoundToAddress
+	 * @see BoundToAddress, BoundToEndpoint
 	 */
 	FTcpSocketBuilder BoundToPort( int32 Port )
 	{
@@ -127,8 +123,7 @@ public:
 	/**
 	 * Sets how long the socket will linger after closing.
 	 *
-	 * @param Timeout - The amount of time to linger before closing.
-	 *
+	 * @param Timeout The amount of time to linger before closing.
 	 * @return This instance (for method chaining).
 	 */
 	FTcpSocketBuilder Lingering( int32 Timeout )
@@ -142,8 +137,7 @@ public:
 	/**
 	 * Sets the socket into a listening state for incoming connections.
 	 *
-	 * @param MaxBacklog - The number of connections to queue before refusing them.
-	 *
+	 * @param MaxBacklog The number of connections to queue before refusing them.
 	 * @return This instance (for method chaining).
 	 */
 	FTcpSocketBuilder Listening( int32 MaxBacklog )
@@ -154,6 +148,39 @@ public:
 		return *this;
 	}
 
+	/**
+	 * Specifies the desired size of the receive buffer in bytes (0 = default).
+	 *
+	 * The socket creation will not fail if the desired size cannot be set or
+	 * if the actual size is less than the desired size.
+	 *
+	 * @param SizeInBytes The size of the buffer.
+	 * @return This instance (for method chaining).
+	 * @see WithSendBufferSize
+	 */
+	FTcpSocketBuilder WithReceiveBufferSize(int32 SizeInBytes)
+	{
+		ReceiveBufferSize = SizeInBytes;
+
+		return *this;
+	}
+
+	/**
+	 * Specifies the desired size of the send buffer in bytes (0 = default).
+	 *
+	 * The socket creation will not fail if the desired size cannot be set or
+	 * if the actual size is less than the desired size.
+	 *
+	 * @param SizeInBytes The size of the buffer.
+	 * @return This instance (for method chaining).
+	 * @see WithReceiveBufferSize
+	 */
+	FTcpSocketBuilder WithSendBufferSize(int32 SizeInBytes)
+	{
+		SendBufferSize = SizeInBytes;
+
+		return *this;
+	}
 
 public:
 
@@ -162,7 +189,7 @@ public:
 	 *
 	 * @return The built socket.
 	 */
-	operator FSocket*( ) const
+	operator FSocket*() const
 	{
 		return Build();
 	}
@@ -172,17 +199,17 @@ public:
 	 *
 	 * @return The built socket.
 	 */
-	FSocket* Build( ) const
+	FSocket* Build() const
 	{
-		FSocket* Socket = NULL;
+		FSocket* Socket = nullptr;
 
 		ISocketSubsystem* SocketSubsystem = ISocketSubsystem::Get(PLATFORM_SOCKETSUBSYSTEM);
 
-		if (SocketSubsystem != NULL)
+		if (SocketSubsystem != nullptr)
 		{
 			Socket = SocketSubsystem->CreateSocket(NAME_Stream, *Description, true);
 
-			if (Socket != NULL)
+			if (Socket != nullptr)
 			{
 				bool Error = !Socket->SetReuseAddr(Reusable) ||
 							 !Socket->SetLinger(Linger, LingerTimeout) ||
@@ -198,13 +225,28 @@ public:
 					Error = Listen && !Socket->Listen(ListenBacklog);
 				}
 
+				if (!Error)
+				{
+					int32 OutNewSize;
+
+					if (ReceiveBufferSize > 0)
+					{
+						Socket->SetReceiveBufferSize(ReceiveBufferSize, OutNewSize);
+					}
+
+					if (SendBufferSize > 0)
+					{
+						Socket->SetSendBufferSize(SendBufferSize, OutNewSize);
+					}
+				}
+
 				if (Error)
 				{
 					GLog->Logf(TEXT("FTcpSocketBuilder: Failed to create the socket %s as configured"), *Description);
 
 					SocketSubsystem->DestroySocket(Socket);
 
-					Socket = NULL;
+					Socket = nullptr;
 				}
 			}
 		}
@@ -212,33 +254,38 @@ public:
 		return Socket;
 	}
 
-
 private:
 
-	// Holds a flag indicating whether socket operations are blocking.
+	/** Holds a flag indicating whether socket operations are blocking. */
 	bool Blocking;
 
-	// Holds a flag indicating whether the socket should be bound.
+	/** Holds a flag indicating whether the socket should be bound. */
 	bool Bound;
 
-	// Holds the IP address (and port) that the socket will be bound to.
+	/** Holds the IP address (and port) that the socket will be bound to. */
 	FIPv4Endpoint BoundEndpoint;
 
-	// Holds the socket's debug description text.
+	/** Holds the socket's debug description text. */
 	FString Description;
 
-	// Holds a flag indicating whether the socket should linger after closing.
+	/** Holds a flag indicating whether the socket should linger after closing. */
 	bool Linger;
 
-	// Holds the amount of time the socket will linger before closing.
+	/** Holds the amount of time the socket will linger before closing. */
 	int32 LingerTimeout;
 
-	// Holds a flag indicating whether the socket should listen for incoming connections.
+	/** Holds a flag indicating whether the socket should listen for incoming connections. */
 	bool Listen;
 
-	// Holds the number of connections to queue up before refusing them.
+	/** Holds the number of connections to queue up before refusing them. */
 	int32 ListenBacklog;
 
-	// Holds a flag indicating whether the bound address can be reused by other sockets.
+	/** The desired size of the receive buffer in bytes (0 = default). */
+	int32 ReceiveBufferSize;
+
+	/** Holds a flag indicating whether the bound address can be reused by other sockets. */
 	bool Reusable;
-};
\ No newline at end of file
+
+	/** The desired size of the send buffer in bytes (0 = default). */
+	int32 SendBufferSize;
+};
diff --git a/Engine/Source/Runtime/Networking/Public/Common/UdpSocketBuilder.h b/Engine/Source/Runtime/Networking/Public/Common/UdpSocketBuilder.h
--- a/Engine/Source/Runtime/Networking/Public/Common/UdpSocketBuilder.h
+++ b/Engine/Source/Runtime/Networking/Public/Common/UdpSocketBuilder.h
@@ -13,7 +13,7 @@ public:
 	/**
 	 * Creates and initializes a new instance.
 	 *
-	 * @param InDescription - Debug description for the socket.
+	 * @param InDescription Debug description for the socket.
 	 */
 	FUdpSocketBuilder( const FString& InDescription )
 		: AllowBroadcast(false)
@@ -23,18 +23,20 @@ public:
 		, Description(InDescription)
 		, MulticastLoopback(false)
 		, MulticastTtl(1)
+		, ReceiveBufferSize(0)
 		, Reusable(false)
+		, SendBufferSize(0)
 	{ }
 
-
 public:
 
 	/**
 	 * Sets socket operations to be blocking.
 	 *
 	 * @return This instance (for method chaining).
+	 * @see AsNonBlocking, AsReusable
 	 */
-	FUdpSocketBuilder AsBlocking( )
+	FUdpSocketBuilder AsBlocking()
 	{
 		Blocking = true;
 
@@ -45,8 +47,9 @@ public:
 	 * Sets socket operations to be non-blocking.
 	 *
 	 * @return This instance (for method chaining).
+	 * @see AsBlocking, AsReusable
 	 */
-	FUdpSocketBuilder AsNonBlocking( )
+	FUdpSocketBuilder AsNonBlocking()
 	{
 		Blocking = false;
 
@@ -57,8 +60,9 @@ public:
 	 * Makes the bound address reusable by other sockets.
 	 *
 	 * @return This instance (for method chaining).
+	 * @see AsBlocking, AsNonBlocking
 	 */
-	FUdpSocketBuilder AsReusable( )
+	FUdpSocketBuilder AsReusable()
 	{
 		Reusable = true;
 
@@ -71,12 +75,9 @@ public:
 	 * Unless specified in a subsequent call to BoundToPort(), a random
 	 * port number will be assigned by the underlying provider.
 	 *
-	 * @param Address - The IP address to bind the socket to.
-	 *
+	 * @param Address The IP address to bind the socket to.
 	 * @return This instance (for method chaining).
-	 *
-	 * @see BoundToEndpoint
-	 * @see BoundToPort
+	 * @see BoundToEndpoint, BoundToPort
 	 */
 	FUdpSocketBuilder BoundToAddress( const FIPv4Address& Address )
 	{
@@ -89,12 +90,9 @@ public:
 	/**
  	 * Sets the local endpoint to bind the socket to.
 	 *
-	 * @param Endpoint - The IP endpoint to bind the socket to.
-	 *
+	 * @param Endpoint The IP endpoint to bind the socket to.
 	 * @return This instance (for method chaining).
-	 *
-	 * @see BoundToAddress
-	 * @see BoundToPort
+	 * @see BoundToAddress, BoundToPort
 	 */
 	FUdpSocketBuilder BoundToEndpoint( const FIPv4Endpoint& Endpoint )
 	{
@@ -110,10 +108,8 @@ public:
 	 * Unless specified in a subsequent call to BoundToAddress(), the local
 	 * address will be determined automatically by the underlying provider.
 	 *
-	 * @param Port - The local port number to bind the socket to.
-	 *
+	 * @param Port The local port number to bind the socket to.
 	 * @return This instance (for method chaining).
-	 *
 	 * @see BoundToAddress
 	 */
 	FUdpSocketBuilder BoundToPort( int32 Port )
@@ -127,9 +123,9 @@ public:
 	/**
 	 * Joins the socket to the specified multicast group.
 	 *
-	 * @param GroupAddress - The IP address of the multicast group to join.
-	 *
+	 * @param GroupAddress The IP address of the multicast group to join.
 	 * @return This instance (for method chaining).
+	 * @see WithMulticastLoopback, WithMulticastTtl
 	 */
 	FUdpSocketBuilder JoinedToGroup( const FIPv4Address& GroupAddress )
 	{
@@ -143,7 +139,7 @@ public:
 	 *
 	 * @return This instance (for method chaining).
 	 */
-	FUdpSocketBuilder WithBroadcast( )
+	FUdpSocketBuilder WithBroadcast()
 	{
 		AllowBroadcast = true;
 
@@ -153,12 +149,12 @@ public:
 	/**
 	 * Enables multicast loopback.
 	 *
-	 * @param AllowLoopback - Whether to allow multicast loopback.
-	 * @param TimeToLive - The time to live.
-	 *
+	 * @param AllowLoopback Whether to allow multicast loopback.
+	 * @param TimeToLive The time to live.
 	 * @return This instance (for method chaining).
+	 * @see JoinedToGroup, WithMulticastTtl
 	 */
-	FUdpSocketBuilder WithMulticastLoopback( )
+	FUdpSocketBuilder WithMulticastLoopback()
 	{
 		MulticastLoopback = true;
 
@@ -168,9 +164,9 @@ public:
 	/**
 	 * Sets the multicast time-to-live.
 	 *
-	 * @param TimeToLive - The time to live.
-	 *
+	 * @param TimeToLive The time to live.
 	 * @return This instance (for method chaining).
+	 * @see JoinedToGroup, WithMulticastLoopback
 	 */
 	FUdpSocketBuilder WithMulticastTtl( uint8 TimeToLive )
 	{
@@ -179,6 +175,39 @@ public:
 		return *this;
 	}
 
+	/**
+	 * Specifies the desired size of the receive buffer in bytes (0 = default).
+	 *
+	 * The socket creation will not fail if the desired size cannot be set or
+	 * if the actual size is less than the desired size.
+	 *
+	 * @param SizeInBytes The size of the buffer.
+	 * @return This instance (for method chaining).
+	 * @see WithSendBufferSize
+	 */
+	FUdpSocketBuilder WithReceiveBufferSize(int32 SizeInBytes)
+	{
+		ReceiveBufferSize = SizeInBytes;
+
+		return *this;
+	}
+
+	/**
+	 * Specifies the desired size of the send buffer in bytes (0 = default).
+	 *
+	 * The socket creation will not fail if the desired size cannot be set or
+	 * if the actual size is less than the desired size.
+	 *
+	 * @param SizeInBytes The size of the buffer.
+	 * @return This instance (for method chaining).
+	 * @see WithReceiveBufferSize
+	 */
+	FUdpSocketBuilder WithSendBufferSize(int32 SizeInBytes)
+	{
+		SendBufferSize = SizeInBytes;
+
+		return *this;
+	}
 
 public:
 
@@ -187,7 +216,7 @@ public:
 	 *
 	 * @return The built socket.
 	 */
-	operator FSocket*( ) const
+	operator FSocket*() const
 	{
 		return Build();
 	}
@@ -197,7 +226,7 @@ public:
 	 *
 	 * @return The built socket.
 	 */
-	FSocket* Build( ) const
+	FSocket* Build() const
 	{
 		FSocket* Socket = nullptr;
 
@@ -238,6 +267,21 @@ public:
 					}
 				}
 
+				if (!Error)
+				{
+					int32 OutNewSize;
+
+					if (ReceiveBufferSize > 0)
+					{
+						Socket->SetReceiveBufferSize(ReceiveBufferSize, OutNewSize);
+					}
+
+					if (SendBufferSize > 0)
+					{
+						Socket->SetSendBufferSize(SendBufferSize, OutNewSize);
+					}
+				}
+
 				if (Error)
 				{
 					GLog->Logf(TEXT("FUdpSocketBuilder: Failed to create the socket %s as configured"), *Description);
@@ -251,33 +295,38 @@ public:
 		return Socket;
 	}
 
-
 private:
 
-	// Holds a flag indicating whether broadcasts will be enabled.
+	/** Holds a flag indicating whether broadcasts will be enabled. */
 	bool AllowBroadcast;
 
-	// Holds a flag indicating whether socket operations are blocking.
+	/** Holds a flag indicating whether socket operations are blocking. */
 	bool Blocking;
 
-	// Holds a flag indicating whether the socket should be bound.
+	/** Holds a flag indicating whether the socket should be bound. */
 	bool Bound;
 
-	// Holds the IP address (and port) that the socket will be bound to.
+	/** Holds the IP address (and port) that the socket will be bound to. */
 	FIPv4Endpoint BoundEndpoint;
 
-	// Holds the socket's debug description text.
+	/** Holds the socket's debug description text. */
 	FString Description;
 
-	// Holds the list of joined multicast groups.
+	/** Holds the list of joined multicast groups. */
 	TArray<FIPv4Address> JoinedGroups;
 
-	// Holds a flag indicating whether multicast loopback will be enabled.
+	/** Holds a flag indicating whether multicast loopback will be enabled. */
 	bool MulticastLoopback;
 
-	// Holds the multicast time to live.
+	/** Holds the multicast time to live. */
 	uint8 MulticastTtl;
 
-	// Holds a flag indicating whether the bound address can be reused by other sockets.
+	/** The desired size of the receive buffer in bytes (0 = default). */
+	int32 ReceiveBufferSize;
+
+	/** Holds a flag indicating whether the bound address can be reused by other sockets. */
 	bool Reusable;
+
+	/** The desired size of the send buffer in bytes (0 = default). */
+	int32 SendBufferSize;
 };
