UE-68794 - Dev-Mobile - Compile UE4Game - use of undeclared identifier 'errno' (2 errors)

diff --git a/Engine/Source/Runtime/Networking/Public/Common/TcpSocketBuilder.h b/Engine/Source/Runtime/Networking/Public/Common/TcpSocketBuilder.h
--- a/Engine/Source/Runtime/Networking/Public/Common/TcpSocketBuilder.h
+++ b/Engine/Source/Runtime/Networking/Public/Common/TcpSocketBuilder.h
@@ -255,7 +255,7 @@ public:
 
 				if (Error)
 				{
-					GLog->Logf(TEXT("FTcpSocketBuilder: Failed to create the socket %s as configured (%i)"), *Description, errno);
+					GLog->Logf(TEXT("FTcpSocketBuilder: Failed to create the socket %s as configured"), *Description);
 
 					SocketSubsystem->DestroySocket(Socket);
 
