Fix bad address creation and incorrect socket protocol construction bug with the TCPSocketBuilder on iOS.

#rnx

diff --git a/Engine/Source/Runtime/Networking/Public/Common/TcpSocketBuilder.h b/Engine/Source/Runtime/Networking/Public/Common/TcpSocketBuilder.h
--- a/Engine/Source/Runtime/Networking/Public/Common/TcpSocketBuilder.h
+++ b/Engine/Source/Runtime/Networking/Public/Common/TcpSocketBuilder.h
@@ -215,7 +215,8 @@ public:
 
 		if (SocketSubsystem != nullptr)
 		{
-			Socket = SocketSubsystem->CreateSocket(NAME_Stream, *Description, ESocketProtocolFamily::IPv4);
+			TSharedRef<FInternetAddr> BoundEndpointAddr = BoundEndpoint.ToInternetAddr();
+			Socket = SocketSubsystem->CreateSocket(NAME_Stream, *Description);
 
 			if (Socket != nullptr)
 			{
@@ -225,7 +226,7 @@ public:
 
 				if (!Error)
 				{
-					Error = Bound && !Socket->Bind(*BoundEndpoint.ToInternetAddrIPV4());
+					Error = Bound && !Socket->Bind(*BoundEndpointAddr);
 				}
 
 				if (!Error)
