Correctly fix the tcpsocketbuilder creation issue with iOS.

#rnx

diff --git a/Engine/Source/Runtime/Networking/Public/Common/TcpSocketBuilder.h b/Engine/Source/Runtime/Networking/Public/Common/TcpSocketBuilder.h
--- a/Engine/Source/Runtime/Networking/Public/Common/TcpSocketBuilder.h
+++ b/Engine/Source/Runtime/Networking/Public/Common/TcpSocketBuilder.h
@@ -216,7 +216,8 @@ public:
 
 		if (SocketSubsystem != nullptr)
 		{
-			Socket = SocketSubsystem->CreateSocket(NAME_Stream, *Description, FNetworkProtocolTypes::IPv4);
+			TSharedRef<FInternetAddr> BoundEndpointAddr = BoundEndpoint.ToInternetAddr();
+			Socket = SocketSubsystem->CreateSocket(NAME_Stream, *Description, BoundEndpointAddr->GetProtocolType());
 
 			if (Socket != nullptr)
 			{
@@ -226,7 +227,7 @@ public:
 
 				if (!Error)
 				{
-					Error = Bound && !Socket->Bind(*BoundEndpoint.ToInternetAddrIPV4());
+					Error = Bound && !Socket->Bind(*BoundEndpointAddr);
 				}
 
 				if (!Error)
