Fix Session Frontend for IPv6-only clients.

diff --git a/Engine/Source/Runtime/Networking/Public/Common/UdpSocketBuilder.h b/Engine/Source/Runtime/Networking/Public/Common/UdpSocketBuilder.h
--- a/Engine/Source/Runtime/Networking/Public/Common/UdpSocketBuilder.h
+++ b/Engine/Source/Runtime/Networking/Public/Common/UdpSocketBuilder.h
@@ -318,9 +318,25 @@ public:
 		}
 		else
 		{
+			TSharedRef<FInternetAddr> MulticastAddress = SocketSubsystem->CreateInternetAddr();
+			MulticastAddress->SetBroadcastAddress();
+			TSharedPtr<FInternetAddr> AddressToUse = nullptr;
+
 			for (const auto& Group : JoinedGroups)
 			{
-				if (!Socket->JoinMulticastGroup(*FIPv4Endpoint(Group.GroupAddress, 0).ToInternetAddr(), *FIPv4Endpoint(Group.InterfaceAddress, 0).ToInternetAddr()))
+				// The Socket code no longer has the multicast address hack handled anymore, as such, we need to properly determine the address ourselves.
+				if (Group.GroupAddress.IsSessionFrontendMulticast() && MulticastAddress->GetProtocolType() != FNetworkProtocolTypes::IPv4)
+				{
+					// This will use the address protocol that we figured out earlier.
+					AddressToUse = MulticastAddress;
+				}
+				else
+				{
+					// Otherwise, we'll use the address we use all the time.
+					AddressToUse = FIPv4Endpoint(Group.GroupAddress, 0).ToInternetAddr();
+				}
+
+				if (!Socket->JoinMulticastGroup(*AddressToUse, *FIPv4Endpoint(Group.InterfaceAddress, 0).ToInternetAddr()))
 				{
 					GLog->Logf(TEXT("FUdpSocketBuilder: Failed to subscribe %s to multicast group %s on interface %s"), *Description, *Group.GroupAddress.ToString(), *Group.InterfaceAddress.ToString());
 					Error = true;
diff --git a/Engine/Source/Runtime/Networking/Public/Interfaces/IPv4/IPv4Address.h b/Engine/Source/Runtime/Networking/Public/Interfaces/IPv4/IPv4Address.h
--- a/Engine/Source/Runtime/Networking/Public/Interfaces/IPv4/IPv4Address.h
+++ b/Engine/Source/Runtime/Networking/Public/Interfaces/IPv4/IPv4Address.h
@@ -234,6 +234,17 @@ public:
 		return ((A == 239) && (B == 255));
 	}
 
+	/**
+	 * Checks whether this IP address is the multicast address we use for the Session Frontend
+	 * Which is 230.0.0.1
+	 *
+	 * @return true if this address matches, false otherwise.
+	 */
+	bool IsSessionFrontendMulticast() const
+	{
+		return (A == 230 && B == 0 && C == 0 && D == 1);
+	}
+
 	/**
 	 * Gets the string representation for this address.
 	 *
