Reduce memory consumption in MessageBus during socket congestion on the receiving end

diff --git a/Engine/Plugins/Messaging/UdpMessaging/Source/UdpMessaging/Private/Transport/UdpMessageTransport.cpp b/Engine/Plugins/Messaging/UdpMessaging/Source/UdpMessaging/Private/Transport/UdpMessageTransport.cpp
--- a/Engine/Plugins/Messaging/UdpMessaging/Source/UdpMessaging/Private/Transport/UdpMessageTransport.cpp
+++ b/Engine/Plugins/Messaging/UdpMessaging/Source/UdpMessaging/Private/Transport/UdpMessageTransport.cpp
@@ -168,6 +168,7 @@ bool FUdpMessageTransport::StartTransport(IMessageTransportHandler& Handler)
 	{
 		MulticastReceiver = new FUdpSocketReceiver(MulticastSocket, ThreadWaitTime, TEXT("UdpMessageMulticastReceiver"));
 		MulticastReceiver->OnDataReceived().BindRaw(this, &FUdpMessageTransport::HandleSocketDataReceived);
+		MulticastReceiver->SetMaxReadBufferSize(2048);
 		MulticastReceiver->Start();
 	}
 
@@ -175,6 +176,7 @@ bool FUdpMessageTransport::StartTransport(IMessageTransportHandler& Handler)
 	UnicastReceiver = new FUdpSocketReceiver(UnicastSocket, ThreadWaitTime, TEXT("UdpMessageUnicastReceiver"));
 	{
 		UnicastReceiver->OnDataReceived().BindRaw(this, &FUdpMessageTransport::HandleSocketDataReceived);
+		UnicastReceiver->SetMaxReadBufferSize(2048);
 		UnicastReceiver->Start();
 	}
 #endif
diff --git a/Engine/Source/Runtime/Networking/Public/Common/UdpSocketReceiver.h b/Engine/Source/Runtime/Networking/Public/Common/UdpSocketReceiver.h
--- a/Engine/Source/Runtime/Networking/Public/Common/UdpSocketReceiver.h
+++ b/Engine/Source/Runtime/Networking/Public/Common/UdpSocketReceiver.h
@@ -68,6 +68,12 @@ public:
 
 public:
 
+	/** Set the maximum size allocated to read off of the socket. */
+	void SetMaxReadBufferSize(uint32 InMaxReadBufferSize)
+	{
+		MaxReadBufferSize = InMaxReadBufferSize;
+	}
+
 	/** Start the receiver thread. */
 	void Start()
 	{
@@ -134,13 +140,14 @@ protected:
 
 		while (Socket->HasPendingData(Size))
 		{
-			FArrayReaderPtr Reader = MakeShareable(new FArrayReader(true));
-			Reader->SetNumUninitialized(FMath::Min(Size, 65507u));
+			FArrayReaderPtr Reader = MakeShared<FArrayReader, ESPMode::ThreadSafe>(true);
+			Reader->SetNumUninitialized(FMath::Min(Size, MaxReadBufferSize));
 
 			int32 Read = 0;
 
 			if (Socket->RecvFrom(Reader->GetData(), Reader->Num(), Read, *Sender))
 			{
+				ensure((uint32)Read < MaxReadBufferSize);
 				Reader->RemoveAt(Read, Reader->Num() - Read, false);
 				DataReceivedDelegate.ExecuteIfBound(Reader, FIPv4Endpoint(Sender));
 			}
@@ -176,6 +183,9 @@ private:
 	/** The amount of time to wait for inbound packets. */
 	FTimespan WaitTime;
 
+	/** The maximum read buffer size used to read the socket. */
+	uint32 MaxReadBufferSize = 65507u;
+
 private:
 
 	/** Holds the data received delegate. */
