#3983 : PR : Increase IPv4 Private Address Space

diff --git a/Engine/Source/Runtime/Networking/Private/Tests/IPv4AddressTest.cpp b/Engine/Source/Runtime/Networking/Private/Tests/IPv4AddressTest.cpp
--- a/Engine/Source/Runtime/Networking/Private/Tests/IPv4AddressTest.cpp
+++ b/Engine/Source/Runtime/Networking/Private/Tests/IPv4AddressTest.cpp
@@ -71,23 +71,25 @@ bool FIPv4AddressTest::RunTest( const FString& Parameters )
 	// site local addresses must be recognized
 	FIPv4Address a8_1 = FIPv4Address(10, 0, 0, 1);
 	FIPv4Address a8_2 = FIPv4Address(172, 16, 0, 1);
-	FIPv4Address a8_3 = FIPv4Address(192, 168, 0, 1);
+	FIPv4Address a8_3 = FIPv4Address(172, 31, 0, 1);
+	FIPv4Address a8_4 = FIPv4Address(192, 168, 0, 1);
 
 	TestTrue(TEXT("10.0.0.1 must be a site local address"), a8_1.IsSiteLocalAddress());
 	TestTrue(TEXT("172.16.0.1 must be a site local address"), a8_2.IsSiteLocalAddress());
-	TestTrue(TEXT("192.168.0.1 must be a site local address"), a8_3.IsSiteLocalAddress());
-
-	FIPv4Address a8_4 = FIPv4Address(11, 0, 0, 1);
-	FIPv4Address a8_5 = FIPv4Address(173, 16, 0, 1);
-	FIPv4Address a8_6 = FIPv4Address(172, 17, 0, 1);
-	FIPv4Address a8_7 = FIPv4Address(193, 168, 0, 1);
-	FIPv4Address a8_8 = FIPv4Address(192, 169, 0, 1);
-
-	TestFalse(TEXT("11.0.0.1 must not be a site local address"), a8_4.IsSiteLocalAddress());
-	TestFalse(TEXT("173.16.0.1 must not be a site local address"), a8_5.IsSiteLocalAddress());
-	TestFalse(TEXT("172.17.0.1 must not be a site local address"), a8_6.IsSiteLocalAddress());
-	TestFalse(TEXT("193.168.0.1 must not be a site local address"), a8_7.IsSiteLocalAddress());
-	TestFalse(TEXT("192.169.0.1 must not be a site local address"), a8_8.IsSiteLocalAddress());
+	TestTrue(TEXT("172.31.0.1 must be a site local address"), a8_3.IsSiteLocalAddress());
+	TestTrue(TEXT("192.168.0.1 must be a site local address"), a8_4.IsSiteLocalAddress());
+
+	FIPv4Address a8_5 = FIPv4Address(11, 0, 0, 1);
+	FIPv4Address a8_6 = FIPv4Address(173, 16, 0, 1);
+	FIPv4Address a8_7 = FIPv4Address(172, 32, 0, 1);
+	FIPv4Address a8_8 = FIPv4Address(193, 168, 0, 1);
+	FIPv4Address a8_9 = FIPv4Address(192, 169, 0, 1);
+
+	TestFalse(TEXT("11.0.0.1 must not be a site local address"), a8_5.IsSiteLocalAddress());
+	TestFalse(TEXT("173.16.0.1 must not be a site local address"), a8_6.IsSiteLocalAddress());
+	TestFalse(TEXT("172.32.0.1 must not be a site local address"), a8_7.IsSiteLocalAddress());
+	TestFalse(TEXT("193.168.0.1 must not be a site local address"), a8_8.IsSiteLocalAddress());
+	TestFalse(TEXT("192.169.0.1 must not be a site local address"), a8_9.IsSiteLocalAddress());
 
 	// host byte order integer constructor
 #if PLATFORM_LITTLE_ENDIAN
diff --git a/Engine/Source/Runtime/Networking/Public/Interfaces/IPv4/IPv4Address.h b/Engine/Source/Runtime/Networking/Public/Interfaces/IPv4/IPv4Address.h
--- a/Engine/Source/Runtime/Networking/Public/Interfaces/IPv4/IPv4Address.h
+++ b/Engine/Source/Runtime/Networking/Public/Interfaces/IPv4/IPv4Address.h
@@ -211,14 +211,14 @@ public:
 	 *
 	 * Site local addresses have one of the following forms:
 	 *		10.x.x.x
-	 *		172.16.x.x
+	 *		172.16-31.x.x
 	 *		192.168.x.x
 	 *
 	 * @return true if it is a site local address, false otherwise.
 	 */
 	bool IsSiteLocalAddress() const
 	{
-		return ((A == 10) || ((A == 172) && (B == 16)) || ((A == 192) && (B == 168)));
+		return ((A == 10) || ((A == 172) && ((B >= 16) && (B <= 31))) || ((A == 192) && (B == 168)));
 	}
 
 	/**
