PR #7882: Hostnames for message bus and file server (Contributed by [REDACTED_PERSON])

diff --git a/Engine/Plugins/Messaging/TcpMessaging/Source/TcpMessaging/Private/TcpMessagingModule.cpp b/Engine/Plugins/Messaging/TcpMessaging/Source/TcpMessaging/Private/TcpMessagingModule.cpp
--- a/Engine/Plugins/Messaging/TcpMessaging/Source/TcpMessaging/Private/TcpMessagingModule.cpp
+++ b/Engine/Plugins/Messaging/TcpMessaging/Source/TcpMessaging/Private/TcpMessagingModule.cpp
@@ -219,7 +219,7 @@ protected:
 		for (FString& ConnectToEndpointString : ConnectToEndpointStrings)
 		{
 			FIPv4Endpoint ConnectToEndpoint;
-			if (FIPv4Endpoint::Parse(ConnectToEndpointString, ConnectToEndpoint) )
+			if (FIPv4Endpoint::Parse(ConnectToEndpointString, ConnectToEndpoint) || FIPv4Endpoint::FromHostAndPort(ConnectToEndpointString, ConnectToEndpoint))
 			{
 				ConnectToEndpoints.Add(ConnectToEndpoint);
 			}
diff --git a/Engine/Source/Runtime/NetworkFile/NetworkFile.Build.cs b/Engine/Source/Runtime/NetworkFile/NetworkFile.Build.cs
--- a/Engine/Source/Runtime/NetworkFile/NetworkFile.Build.cs
+++ b/Engine/Source/Runtime/NetworkFile/NetworkFile.Build.cs
@@ -12,6 +12,7 @@ public class NetworkFile : ModuleRules
 			new string[] {
 				"Core",
 				"CoreUObject",
+				"Networking",
 				"Sockets"
 			});
 
diff --git a/Engine/Source/Runtime/NetworkFile/Private/TCPTransport.cpp b/Engine/Source/Runtime/NetworkFile/Private/TCPTransport.cpp
--- a/Engine/Source/Runtime/NetworkFile/Private/TCPTransport.cpp
+++ b/Engine/Source/Runtime/NetworkFile/Private/TCPTransport.cpp
@@ -7,6 +7,7 @@
 #include "Sockets.h"
 #include "MultichannelTcpSocket.h"
 #include "NetworkPlatformFile.h"
+#include "Interfaces/IPv4/IPv4Endpoint.h"
 
 FTCPTransport::FTCPTransport()
 	:FileSocket(NULL)
@@ -29,8 +30,18 @@ bool FTCPTransport::Initialize(const TCHAR* InHostIp)
 	TSharedRef<FInternetAddr> Addr = SSS->CreateInternetAddr();
 	bool bIsValid;
 
-	Addr->SetIp(*HostIp, bIsValid);
-	Addr->SetPort(DEFAULT_TCP_FILE_SERVING_PORT);
+	FIPv4Endpoint FileServerEndpoint;
+	if (FIPv4Endpoint::Parse(InHostIp, FileServerEndpoint) || FIPv4Endpoint::FromHostAndPort(InHostIp, FileServerEndpoint))
+	{
+		Addr->SetIp(FileServerEndpoint.Address.Value);
+		Addr->SetPort(FileServerEndpoint.Port == 0 ? DEFAULT_TCP_FILE_SERVING_PORT : FileServerEndpoint.Port);
+		bIsValid = true;
+	}
+	else
+	{
+		Addr->SetIp(*HostIp, bIsValid);
+		Addr->SetPort(DEFAULT_TCP_FILE_SERVING_PORT);
+	}
 
 	if (bIsValid)
 	{
@@ -113,4 +124,3 @@ FTCPTransport::~FTCPTransport()
 	ISocketSubsystem::Get()->DestroySocket(FileSocket);
 	FileSocket = NULL;
 }
-
diff --git a/Engine/Source/Runtime/Networking/Private/IPv4/IPv4Endpoint.cpp b/Engine/Source/Runtime/Networking/Private/IPv4/IPv4Endpoint.cpp
--- a/Engine/Source/Runtime/Networking/Private/IPv4/IPv4Endpoint.cpp
+++ b/Engine/Source/Runtime/Networking/Private/IPv4/IPv4Endpoint.cpp
@@ -1,6 +1,7 @@
 // Copyright Epic Games, Inc. All Rights Reserved.
 
 #include "Interfaces/IPv4/IPv4Endpoint.h"
+#include "AddressInfoTypes.h"
 
 
 /* FIPv4Endpoint static initialization
@@ -44,3 +45,35 @@ bool FIPv4Endpoint::Parse(const FString& EndpointString, FIPv4Endpoint& OutEndpo
 
 	return false;
 }
+
+bool FIPv4Endpoint::FromHostAndPort(const FString& HostAndPortString, FIPv4Endpoint& OutEndpoint)
+{
+	FString Host;
+	FString Port;
+	if (!HostAndPortString.Split(TEXT(":"), &Host, &Port, ESearchCase::IgnoreCase, ESearchDir::FromEnd))
+	{
+		return false;
+	}
+
+	Host.TrimStartAndEndInline();
+	Port.TrimStartAndEndInline();
+	if(Host.IsEmpty() || Port.IsEmpty())
+	{
+		return false;
+	}
+
+	const FAddressInfoResult AddressInfo = ISocketSubsystem::Get(PLATFORM_SOCKETSUBSYSTEM)->GetAddressInfo(
+		*Host,
+		*Port,
+		EAddressInfoFlags::NoResolveService,
+		FNetworkProtocolTypes::IPv4,
+		SOCKTYPE_Unknown);
+	if (AddressInfo.ReturnCode != SE_NO_ERROR || AddressInfo.Results.Num() == 0)
+	{
+		return false;
+	}
+
+	const FAddressInfoResultData& Result = AddressInfo.Results[0];
+	OutEndpoint = FIPv4Endpoint(Result.Address);
+	return true;
+}
diff --git a/Engine/Source/Runtime/Networking/Public/Interfaces/IPv4/IPv4Endpoint.h b/Engine/Source/Runtime/Networking/Public/Interfaces/IPv4/IPv4Endpoint.h
--- a/Engine/Source/Runtime/Networking/Public/Interfaces/IPv4/IPv4Endpoint.h
+++ b/Engine/Source/Runtime/Networking/Public/Interfaces/IPv4/IPv4Endpoint.h
@@ -185,6 +185,15 @@ public:
 	 */
 	static NETWORKING_API bool Parse(const FString& EndpointString, FIPv4Endpoint& OutEndpoint);
 
+	/**
+	 * Resolves a hostname and port to an IPv4 endpoint.
+	 *
+	 * @param HostAndPortString The host and port to resolve.
+	 * @param OutEndpoint Will contain the resolved endpoint.
+	 * @return true if the host and port were parsed and resolved successfully, false otherwise.
+	 */
+	static NETWORKING_API bool FromHostAndPort(const FString& HostAndPortString, FIPv4Endpoint& OutEndpoint);
+
 public:
 
 	/** Defines the wild card endpoint, which is 0.0.0.0:0 */
