Sockets: Fixes misleading unhandled socket error and subsequent crash.

The issue was that when the udp socket builder got a socket error (e.g. bind failed), prior to calling GetLastError it was calling GLog->Logf, which under certain circumstances can also cause an error which GetLastError would then return (tends to happen when -log is in the command line), which may not even be a valid socket error code, therefore being flagged as an unhandled socket error and ultimately crashing with check(0).

The fix is to call GetLastError before calling any other function that can clobber the error code that GetLastError returns.

Also added a warning message when the receive or send buffer sizes fail to be set, but it will still return the otherwise configured socket.

diff --git a/Engine/Source/Runtime/Networking/Public/Common/UdpSocketBuilder.h b/Engine/Source/Runtime/Networking/Public/Common/UdpSocketBuilder.h
--- a/Engine/Source/Runtime/Networking/Public/Common/UdpSocketBuilder.h
+++ b/Engine/Source/Runtime/Networking/Public/Common/UdpSocketBuilder.h
@@ -283,100 +283,121 @@ public:
 			GLog->Logf(TEXT("FUdpSocketBuilder: Failed to create socket %s"), *Description);
 			return nullptr;
 		}
-		
+
 		// configure socket
-		bool Error =
-			!Socket->SetNonBlocking(!Blocking) ||
+
+		if (!Socket->SetNonBlocking(!Blocking) ||
 			!Socket->SetReuseAddr(Reusable) ||
 			!Socket->SetBroadcast(AllowBroadcast) ||
-			!Socket->SetRecvErr();
-
-		// bind socket
-		if (Error)
+			!Socket->SetRecvErr())
 		{
-			GLog->Logf(TEXT("FUdpSocketBuilder: Failed to configure %s (blocking: %i, reusable: %i, broadcast: %i)"), *Description, Blocking, Reusable, AllowBroadcast);
+			const ESocketErrors SocketError = SocketSubsystem->GetLastErrorCode();
+
+			GLog->Logf(TEXT("FUdpSocketBuilder: Failed to configure %s (blocking: %i, reusable: %i, broadcast: %i). Error code %d."), 
+				*Description, Blocking, Reusable, AllowBroadcast, int32(SocketError));
+
+			SocketSubsystem->DestroySocket(Socket);
+			return nullptr;
 		}
-		else
+
+		// bind socket
+
+		if (Bound && !Socket->Bind(*RemoteAddr))
 		{
-			Error = Bound && !Socket->Bind(*RemoteAddr);
+			const ESocketErrors SocketError = SocketSubsystem->GetLastErrorCode();
+
+			GLog->Logf(TEXT("FUdpSocketBuilder: Failed to bind %s to %s. Error code %d."),
+				*Description, *BoundEndpoint.ToString(), int32(SocketError));
+
+			SocketSubsystem->DestroySocket(Socket);
+			return nullptr;
 		}
 
 		// configure multicast
-		if (Error)
-		{
-			GLog->Logf(TEXT("FUdpSocketBuilder: Failed to bind %s to %s"), *Description, *BoundEndpoint.ToString());
-		}
-		else
+
+		if (!Socket->SetMulticastLoopback(MulticastLoopback) || !Socket->SetMulticastTtl(MulticastTtl))
 		{
-			Error = !Socket->SetMulticastLoopback(MulticastLoopback) || !Socket->SetMulticastTtl(MulticastTtl);
+			const ESocketErrors SocketError = SocketSubsystem->GetLastErrorCode();
+
+			GLog->Logf(TEXT("FUdpSocketBuilder: Failed to configure multicast for %s (loopback: %i, ttl: %i). Error code %d."),
+				*Description, MulticastLoopback, MulticastTtl, int32(SocketError));
+
+			SocketSubsystem->DestroySocket(Socket);
+			return nullptr;
 		}
 
 		// join multicast groups
-		if (Error)
-		{
-			GLog->Logf(TEXT("FUdpSocketBuilder: Failed to configure multicast for %s (loopback: %i, ttl: %i)"), *Description, MulticastLoopback, MulticastTtl);
-		}
-		else
+
+		TSharedRef<FInternetAddr> MulticastAddress = SocketSubsystem->CreateInternetAddr(RemoteAddr->GetProtocolType());
+		MulticastAddress->SetBroadcastAddress();
+		TSharedPtr<FInternetAddr> AddressToUse = nullptr;
+
+		for (const auto& Group : JoinedGroups)
 		{
-			TSharedRef<FInternetAddr> MulticastAddress = SocketSubsystem->CreateInternetAddr(RemoteAddr->GetProtocolType());
-			MulticastAddress->SetBroadcastAddress();
-			TSharedPtr<FInternetAddr> AddressToUse = nullptr;
+			// The Socket code no longer has the multicast address hack handled anymore, as such, we need to properly determine the address ourselves.
+			if (Group.GroupAddress.IsSessionFrontendMulticast() && MulticastAddress->GetProtocolType() != FNetworkProtocolTypes::IPv4)
+			{
+				// This will use the address protocol that we figured out earlier.
+				AddressToUse = MulticastAddress;
+			}
+			else
+			{
+				// Otherwise, we'll use the address we use all the time.
+				AddressToUse = FIPv4Endpoint(Group.GroupAddress, 0).ToInternetAddr();
+			}
 
-			for (const auto& Group : JoinedGroups)
+			if (!Socket->JoinMulticastGroup(*AddressToUse, *FIPv4Endpoint(Group.InterfaceAddress, 0).ToInternetAddr()))
 			{
-				// The Socket code no longer has the multicast address hack handled anymore, as such, we need to properly determine the address ourselves.
-				if (Group.GroupAddress.IsSessionFrontendMulticast() && MulticastAddress->GetProtocolType() != FNetworkProtocolTypes::IPv4)
-				{
-					// This will use the address protocol that we figured out earlier.
-					AddressToUse = MulticastAddress;
-				}
-				else
-				{
-					// Otherwise, we'll use the address we use all the time.
-					AddressToUse = FIPv4Endpoint(Group.GroupAddress, 0).ToInternetAddr();
-				}
+				const ESocketErrors SocketError = SocketSubsystem->GetLastErrorCode();
 
-				if (!Socket->JoinMulticastGroup(*AddressToUse, *FIPv4Endpoint(Group.InterfaceAddress, 0).ToInternetAddr()))
-				{
-					GLog->Logf(TEXT("FUdpSocketBuilder: Failed to subscribe %s to multicast group %s on interface %s"), *Description, *Group.GroupAddress.ToString(), *Group.InterfaceAddress.ToString());
-					Error = true;
+				GLog->Logf(TEXT("FUdpSocketBuilder: Failed to subscribe %s to multicast group %s on interface %s. Error code %d."),
+					*Description, *Group.GroupAddress.ToString(), *Group.InterfaceAddress.ToString(), int32(SocketError));
 
-					break;
-				}
+				SocketSubsystem->DestroySocket(Socket);
+				return nullptr;
 			}
 		}
 
 		// set multicast outgoing interface
-		if (!Error && MulticastInterface != FIPv4Address::Any)
+		if (MulticastInterface != FIPv4Address::Any)
 		{
-			Error = !Socket->SetMulticastInterface(*FIPv4Endpoint(MulticastInterface, 0).ToInternetAddr());
-			if (Error)
+			if (!Socket->SetMulticastInterface(*FIPv4Endpoint(MulticastInterface, 0).ToInternetAddr()))
 			{
-				GLog->Logf(TEXT("FUdpSocketBuilder: Failed to set multicast outgoing interface for %s to %s "), *Description, *MulticastInterface.ToString());
+				const ESocketErrors SocketError = SocketSubsystem->GetLastErrorCode();
+
+				GLog->Logf(TEXT("FUdpSocketBuilder: Failed to set multicast outgoing interface for %s to %s. Error code %d."),
+					*Description, *MulticastInterface.ToString(), int32(SocketError));
+
+				SocketSubsystem->DestroySocket(Socket);
+				return nullptr;
 			}
 		}
 
 		// set buffer sizes
-		if (!Error)
 		{
 			int32 OutNewSize;
 
 			if (ReceiveBufferSize > 0)
 			{
-				Socket->SetReceiveBufferSize(ReceiveBufferSize, OutNewSize);
+				if (!Socket->SetReceiveBufferSize(ReceiveBufferSize, OutNewSize))
+				{
+					const ESocketErrors SocketError = SocketSubsystem->GetLastErrorCode();
+
+					GLog->Logf(TEXT("FUdpSocketBuilder: Warning - could not set receive buffer size to %d for %s. Error code %d."),
+						ReceiveBufferSize, *Description, int32(SocketError));
+				}
 			}
 
 			if (SendBufferSize > 0)
 			{
-				Socket->SetSendBufferSize(SendBufferSize, OutNewSize);
-			}
-		}
+				if (!Socket->SetSendBufferSize(SendBufferSize, OutNewSize))
+				{
+					const ESocketErrors SocketError = SocketSubsystem->GetLastErrorCode();
 
-		if (Error)
-		{
-			GLog->Logf(TEXT("FUdpSocketBuilder: Failed to create and initialize socket %s (last error: %i)"), *Description, (int32)SocketSubsystem->GetLastErrorCode());
-			SocketSubsystem->DestroySocket(Socket);
-			Socket = nullptr;
+					GLog->Logf(TEXT("FUdpSocketBuilder: Warning - could not set send buffer size to %d for %s. Error code %d."),
+						SendBufferSize, *Description, int32(SocketError));
+				}
+			}
 		}
 
 		return Socket;
