fix small truncation warning in IPv4Endpoint

diff --git a/Engine/Source/Runtime/Networking/Public/Interfaces/IPv4/IPv4Endpoint.h b/Engine/Source/Runtime/Networking/Public/Interfaces/IPv4/IPv4Endpoint.h
--- a/Engine/Source/Runtime/Networking/Public/Interfaces/IPv4/IPv4Endpoint.h
+++ b/Engine/Source/Runtime/Networking/Public/Interfaces/IPv4/IPv4Endpoint.h
@@ -61,7 +61,7 @@ public:
 		int32 OutPort;
 		InternetAddr->GetIp(Address.Value);
 		InternetAddr->GetPort(OutPort);
-		Port = OutPort;
+		Port = static_cast<uint16>(OutPort);
 	}
 
 public:
