Networking: Made UdpSocketReceiver::bStopping flag atomic because it is written and read from different threads.

It is read from the FRunnable thread, but written into in Stop() from another.

Also added the missing 'b' prefix.

#rnx

diff --git a/Engine/Source/Runtime/Networking/Public/Common/UdpSocketReceiver.h b/Engine/Source/Runtime/Networking/Public/Common/UdpSocketReceiver.h
--- a/Engine/Source/Runtime/Networking/Public/Common/UdpSocketReceiver.h
+++ b/Engine/Source/Runtime/Networking/Public/Common/UdpSocketReceiver.h
@@ -13,6 +13,8 @@
 #include "Interfaces/IPv4/IPv4Endpoint.h"
 #include "IPAddress.h"
 
+#include <atomic>
+
 /**
  * Temporary fix for concurrency crashes. This whole class will be redesigned.
  */
@@ -45,7 +47,7 @@ public:
 	 */
 	FUdpSocketReceiver(FSocket* InSocket, const FTimespan& InWaitTime, const TCHAR* InThreadName)
 		: Socket(InSocket)
-		, Stopping(false)
+		, bStopping(false)
 		, Thread(nullptr)
 		, ThreadName(InThreadName)
 		, WaitTime(InWaitTime)
@@ -110,7 +112,7 @@ public:
 
 	virtual uint32 Run() override
 	{
-		while (!Stopping)
+		while (!bStopping.load(std::memory_order_acquire))
 		{
 			Update(WaitTime);
 		}
@@ -120,7 +122,7 @@ public:
 
 	virtual void Stop() override
 	{
-		Stopping = true;
+		bStopping.store(true, std::memory_order_release);
 	}
 
 	virtual void Exit() override { }
@@ -173,7 +175,7 @@ private:
 	ISocketSubsystem* SocketSubsystem;
 
 	/** Flag indicating that the thread is stopping. */
-	bool Stopping;
+	std::atomic<bool> bStopping;
 
 	/** The thread object. */
 	FRunnableThread* Thread;
